Also fix unrecognized tab, and duplicate struct link in both registry and data table editor

diff --git a/Engine/Plugins/Runtime/DataRegistry/Source/DataRegistryEditor/Private/DataRegistryEditorToolkit.cpp b/Engine/Plugins/Runtime/DataRegistry/Source/DataRegistryEditor/Private/DataRegistryEditorToolkit.cpp
--- a/Engine/Plugins/Runtime/DataRegistry/Source/DataRegistryEditor/Private/DataRegistryEditorToolkit.cpp
+++ b/Engine/Plugins/Runtime/DataRegistry/Source/DataRegistryEditor/Private/DataRegistryEditorToolkit.cpp
@@ -132,7 +132,7 @@ void FDataRegistryEditorToolkit::CreateAndRegisterPropertiesTab(const TSharedRef
 
 void FDataRegistryEditorToolkit::InitDataRegistryEditor(const EToolkitMode::Type Mode, const TSharedPtr< class IToolkitHost >& InitToolkitHost, UDataRegistry* Registry)
 {
-	TSharedRef<FTabManager::FLayout> StandaloneDefaultLayout = FTabManager::NewLayout("Standalone_CompositeDataRegistryEditor_temp_Layout")
+	TSharedRef<FTabManager::FLayout> StandaloneDefaultLayout = FTabManager::NewLayout("Standalone_DataRegistryEditor_Layout")
 		->AddArea
 		(
 			FTabManager::NewPrimaryArea()->SetOrientation(Orient_Horizontal)
@@ -154,13 +154,6 @@ void FDataRegistryEditorToolkit::InitDataRegistryEditor(const EToolkitMode::Type
 				FTabManager::NewSplitter()
 				->SetOrientation(Orient_Vertical)
 				->Split
-				(
-					FTabManager::NewStack()
-					->SetSizeCoefficient(0.1f)
-					->SetHideTabWell(true)
-					->AddTab(GetToolbarTabId(), ETabState::OpenedTab)
-				)
-				->Split
 				(
 					FTabManager::NewStack()
 					->AddTab(RowListTabId, ETabState::OpenedTab)
@@ -634,20 +627,8 @@ void FDataRegistryEditorToolkit::PostRegenerateMenusAndToolbars()
 				.ForegroundColor(FSlateColor::UseForeground())
 				[
 					SNew(SImage)
-					.Image(FEditorStyle::GetBrush("PropertyWindow.Button_Browse"))
+					.Image(FEditorStyle::GetBrush("Icons.Search"))
 				]
-			]
-			+SHorizontalBox::Slot()
-			.AutoWidth()
-			.VAlign(VAlign_Center)
-			.Padding(0.0f, 0.0f, 8.0f, 0.0f)
-			[
-				SNew(SHyperlink)
-				.Style(FEditorStyle::Get(), "Common.GotoNativeCodeHyperlink")
-				.Visibility(!UDS ? EVisibility::Visible : EVisibility::Collapsed)
-				.OnNavigate(this, &FDataRegistryEditorToolkit::OnNavigateToDataRegistryRowCode)
-				.Text(FText::FromString(GetNameSafe(DataRegistry->GetItemStruct())))
-				.ToolTipText(FText::Format(LOCTEXT("GoToCode_ToolTip", "Click to open this source file in {0}"), FSourceCodeNavigation::GetSelectedSourceCodeIDE()))
 			];
 	
 		SetMenuOverlay(MenuOverlayBox);
@@ -674,16 +655,10 @@ void FDataRegistryEditorToolkit::OnEditDataTableStructClicked()
 	if (DataRegistry && DataRegistry->GetItemStruct())
 	{
 		GEditor->GetEditorSubsystem<UAssetEditorSubsystem>()->OpenEditorForAsset(DataRegistry->GetItemStruct()->GetPathName());
-		FSourceCodeNavigation::NavigateToStruct(DataRegistry->GetItemStruct());
-	}
-}
-
-void FDataRegistryEditorToolkit::OnNavigateToDataRegistryRowCode()
-{
-	const UDataRegistry* DataRegistry = GetDataRegistry();
-	if (DataRegistry && FSourceCodeNavigation::CanNavigateToStruct(DataRegistry->GetItemStruct()))
-	{
-		FSourceCodeNavigation::NavigateToStruct(DataRegistry->GetItemStruct());
+		if (FSourceCodeNavigation::CanNavigateToStruct(DataRegistry->GetItemStruct()))
+		{
+			FSourceCodeNavigation::NavigateToStruct(DataRegistry->GetItemStruct());
+		}
 	}
 }
 
diff --git a/Engine/Plugins/Runtime/DataRegistry/Source/DataRegistryEditor/Private/DataRegistryEditorToolkit.h b/Engine/Plugins/Runtime/DataRegistry/Source/DataRegistryEditor/Private/DataRegistryEditorToolkit.h
--- a/Engine/Plugins/Runtime/DataRegistry/Source/DataRegistryEditor/Private/DataRegistryEditorToolkit.h
+++ b/Engine/Plugins/Runtime/DataRegistry/Source/DataRegistryEditor/Private/DataRegistryEditorToolkit.h
@@ -70,7 +70,6 @@ protected:
 
 	FReply OnFindRowInContentBrowserClicked();
 	void OnEditDataTableStructClicked();
-	void OnNavigateToDataRegistryRowCode();
 
 	FText GetCellText(FDataTableEditorRowListViewDataPtr InRowDataPointer, int32 ColumnIndex) const;
 	FText GetCellToolTipText(FDataTableEditorRowListViewDataPtr InRowDataPointer, int32 ColumnIndex) const;
diff --git a/Engine/Source/Editor/DataTableEditor/Private/DataTableEditor.cpp b/Engine/Source/Editor/DataTableEditor/Private/DataTableEditor.cpp
--- a/Engine/Source/Editor/DataTableEditor/Private/DataTableEditor.cpp
+++ b/Engine/Source/Editor/DataTableEditor/Private/DataTableEditor.cpp
@@ -487,13 +487,15 @@ void FDataTableEditor::OnEditDataTableStructClicked()
 	const UDataTable* DataTable = GetDataTable();
 	if (DataTable)
 	{
-
 		const UScriptStruct* ScriptStruct = DataTable->GetRowStruct();
 
 		if (ScriptStruct)
 		{
 			GEditor->GetEditorSubsystem<UAssetEditorSubsystem>()->OpenEditorForAsset(ScriptStruct->GetPathName());
-			FSourceCodeNavigation::NavigateToStruct(ScriptStruct);
+			if (FSourceCodeNavigation::CanNavigateToStruct(ScriptStruct))
+			{
+				FSourceCodeNavigation::NavigateToStruct(ScriptStruct);
+			}
 		}
 	}
 }
@@ -969,18 +971,6 @@ void FDataTableEditor::PostRegenerateMenusAndToolbars()
 					SNew(SImage)
 					.Image(FEditorStyle::GetBrush("Icons.Search"))
 				]
-			]
-			+SHorizontalBox::Slot()
-			.AutoWidth()
-			.VAlign(VAlign_Center)
-			.Padding(0.0f, 0.0f, 8.0f, 0.0f)
-			[
-				SNew(SHyperlink)
-				.Style(FEditorStyle::Get(), "Common.GotoNativeCodeHyperlink")
-				.Visibility(!UDS ? EVisibility::Visible : EVisibility::Collapsed)
-				.OnNavigate(this, &FDataTableEditor::OnNavigateToDataTableRowCode)
-				.Text(FText::FromName(DataTable->GetRowStructName()))
-				.ToolTipText(FText::Format(LOCTEXT("GoToCode_ToolTip", "Click to open this source file in {0}"), FSourceCodeNavigation::GetSelectedSourceCodeIDE()))
 			];
 	
 		SetMenuOverlay(MenuOverlayBox);
@@ -1000,15 +990,6 @@ FReply FDataTableEditor::OnFindRowInContentBrowserClicked()
 	return FReply::Handled();
 }
 
-void FDataTableEditor::OnNavigateToDataTableRowCode()
-{
-	const UDataTable* DataTable = GetDataTable();
-	if (DataTable && FSourceCodeNavigation::CanNavigateToStruct(DataTable->GetRowStruct()))
-	{
-		FSourceCodeNavigation::NavigateToStruct(DataTable->GetRowStruct());
-	}
-}
-
 void FDataTableEditor::RefreshCachedDataTable(const FName InCachedSelection, const bool bUpdateEvenIfValid)
 {
 	UDataTable* Table = GetEditableDataTable();
diff --git a/Engine/Source/Editor/DataTableEditor/Private/DataTableEditor.h b/Engine/Source/Editor/DataTableEditor/Private/DataTableEditor.h
--- a/Engine/Source/Editor/DataTableEditor/Private/DataTableEditor.h
+++ b/Engine/Source/Editor/DataTableEditor/Private/DataTableEditor.h
@@ -97,8 +97,6 @@ protected:
 
 	FReply OnFindRowInContentBrowserClicked();
 
-	void OnNavigateToDataTableRowCode();
-
 	FText GetCellText(FDataTableEditorRowListViewDataPtr InRowDataPointer, int32 ColumnIndex) const;
 	FText GetCellToolTipText(FDataTableEditorRowListViewDataPtr InRowDataPointer, int32 ColumnIndex) const;
 
