[NO NON-PII COMMIT MESSAGE]

diff --git a/Engine/Plugins/Experimental/GameFeatures/Source/GameFeatures/Private/GameFeatureAction_DataRegistrySource.cpp b/Engine/Plugins/Experimental/GameFeatures/Source/GameFeatures/Private/GameFeatureAction_DataRegistrySource.cpp
--- a/Engine/Plugins/Experimental/GameFeatures/Source/GameFeatures/Private/GameFeatureAction_DataRegistrySource.cpp
+++ b/Engine/Plugins/Experimental/GameFeatures/Source/GameFeatures/Private/GameFeatureAction_DataRegistrySource.cpp
@@ -77,12 +77,18 @@ void UGameFeatureAction_DataRegistrySource::OnGameFeatureDeactivating(FGameFeatu
 		{
 			if (!RegistrySource.DataTableToAdd.IsNull())
 			{
-				DataRegistrySubsystem->UnregisterSpecificAsset(RegistrySource.RegistryToAddTo, RegistrySource.DataTableToAdd.ToSoftObjectPath());
+				if (!DataRegistrySubsystem->UnregisterSpecificAsset(RegistrySource.RegistryToAddTo, RegistrySource.DataTableToAdd.ToSoftObjectPath()))
+				{
+					UE_LOG(LogGameFeatures, Log, TEXT("OnGameFeatureDeactivating %s: DataRegistry data table %s failed to unregister"), *GetPathName(), *RegistrySource.DataTableToAdd.ToString())
+				}
 			}
 
 			if (!RegistrySource.CurveTableToAdd.IsNull())
 			{
-				DataRegistrySubsystem->UnregisterSpecificAsset(RegistrySource.RegistryToAddTo, RegistrySource.CurveTableToAdd.ToSoftObjectPath());
+				if (!DataRegistrySubsystem->UnregisterSpecificAsset(RegistrySource.RegistryToAddTo, RegistrySource.CurveTableToAdd.ToSoftObjectPath()))
+				{
+					UE_LOG(LogGameFeatures, Log, TEXT("OnGameFeatureDeactivating %s: DataRegistry curve table %s failed to unregister"), *GetPathName(), *RegistrySource.CurveTableToAdd.ToString())
+				}
 			}
 		}
 	}
diff --git a/Engine/Plugins/Runtime/DataRegistry/Source/DataRegistry/Private/DataRegistrySubsystem.cpp b/Engine/Plugins/Runtime/DataRegistry/Source/DataRegistry/Private/DataRegistrySubsystem.cpp
--- a/Engine/Plugins/Runtime/DataRegistry/Source/DataRegistry/Private/DataRegistrySubsystem.cpp
+++ b/Engine/Plugins/Runtime/DataRegistry/Source/DataRegistry/Private/DataRegistrySubsystem.cpp
@@ -658,17 +658,24 @@ void UDataRegistrySubsystem::ApplyPreregisterMap(UDataRegistry* Registry)
 	{
 		for (int32 i = 0; i < FoundPreregister->Num(); i++)
 		{
+			bool bRegistered = false;
 			const FSoftObjectPath& AssetPath = (*FoundPreregister)[i].Key;
 			FAssetData AssetData;
 			if (AssetManager.GetAssetDataForPath(AssetPath, AssetData))
 			{
-				Registry->RegisterSpecificAsset(AssetData, (*FoundPreregister)[i].Value);
+				bRegistered = Registry->RegisterSpecificAsset(AssetData, (*FoundPreregister)[i].Value);
 			}
 			else if (Settings->CanIgnoreMissingAssetData())
 			{
 				// Construct fake asset data and register that
 				AssetData = FAssetData(AssetPath.GetLongPackageName(), AssetPath.GetAssetPathString(), NAME_Object);
-				Registry->RegisterSpecificAsset(AssetData, (*FoundPreregister)[i].Value);
+				bRegistered = Registry->RegisterSpecificAsset(AssetData, (*FoundPreregister)[i].Value);
+			}
+
+			if (!bRegistered)
+			{
+				// If specific type is mentioned, it is expected to always succeed
+				UE_LOG(LogDataRegistry, Warning, TEXT("ApplyPreregisterMap failed to register %s with %s, there needs to be a meta source that handles registered assets with matching data"), *AssetPath.ToString(), *Registry->GetRegistryType().ToString());
 			}
 		}
 	}
