It now leaves the invalid data but puts up a warning image

diff --git a/Engine/Plugins/Runtime/DataRegistry/Source/DataRegistryEditor/Private/DataRegistryIdCustomization.cpp b/Engine/Plugins/Runtime/DataRegistry/Source/DataRegistryEditor/Private/DataRegistryIdCustomization.cpp
--- a/Engine/Plugins/Runtime/DataRegistry/Source/DataRegistryEditor/Private/DataRegistryIdCustomization.cpp
+++ b/Engine/Plugins/Runtime/DataRegistry/Source/DataRegistryEditor/Private/DataRegistryIdCustomization.cpp
@@ -48,26 +48,41 @@ void SDataRegistryItemNameWidget::Construct(const FArguments& InArgs)
 
 	this->ChildSlot
 	[
-		SAssignNew(ValueSwitcher, SWidgetSwitcher)
-		.WidgetIndex(0)
-		+ SWidgetSwitcher::Slot()
-		[	
-			PropertyCustomizationHelpers::MakePropertyComboBox(NameArgs)
-		]
-		+ SWidgetSwitcher::Slot()
+		SNew(SHorizontalBox)
+		+ SHorizontalBox::Slot()
+		.FillWidth(1)
 		[
-			SNew(SComboButton)
-			.OnGetMenuContent(this, &SDataRegistryItemNameWidget::GetTagContent)
-			.OnMenuOpenChanged(this, &SDataRegistryItemNameWidget::OnTagUIOpened)
-			.ContentPadding(FMargin(4.0f, 2.0f))
-			.MenuPlacement(MenuPlacement_BelowAnchor)
-			.ButtonContent()
+			SAssignNew(ValueSwitcher, SWidgetSwitcher)
+			.WidgetIndex(0)
+			+ SWidgetSwitcher::Slot()
+			[	
+				PropertyCustomizationHelpers::MakePropertyComboBox(NameArgs)
+			]
+			+ SWidgetSwitcher::Slot()
 			[
-				SNew(STextBlock)
-				.Text(this, &SDataRegistryItemNameWidget::OnGetNameValueText)
-				.Font(FEditorStyle::GetFontStyle(TEXT("PropertyWindow.NormalFont")))
+				SNew(SComboButton)
+				.OnGetMenuContent(this, &SDataRegistryItemNameWidget::GetTagContent)
+				.OnMenuOpenChanged(this, &SDataRegistryItemNameWidget::OnTagUIOpened)
+				.MenuPlacement(MenuPlacement_BelowAnchor)
+				.ButtonContent()
+				[
+					SNew(STextBlock)
+					.Text(this, &SDataRegistryItemNameWidget::OnGetNameValueText)
+					.Font(FEditorStyle::GetFontStyle(TEXT("PropertyWindow.NormalFont")))
+				]
 			]
 		]
+		+ SHorizontalBox::Slot()
+		.Padding(FMargin(2, 0, 0, 0))
+		.AutoWidth()
+		.VAlign(EVerticalAlignment::VAlign_Center)
+		[
+			SNew(SImage)
+			.Image(FAppStyle::Get().GetBrush("Icons.WarningWithColor"))
+			.DesiredSizeOverride(FVector2D(16, 16))
+			.ToolTipText(LOCTEXT("InvalidRegistryId", "Unknown Item Name, this will not work at runtime with the current registry settings"))
+			.Visibility(this, &SDataRegistryItemNameWidget::GetWarningVisibility)
+		]
 	];
 }
 
@@ -79,8 +94,9 @@ void SDataRegistryItemNameWidget::Tick(const FGeometry& AllottedGeometry, const
 	// Refresh if type has changed, or it's a custom type that might be different based on external factors
 	if (CurrentId.RegistryType != CachedIdValue.RegistryType || CurrentId.RegistryType == FDataRegistryType::CustomContextType)
 	{
+		bool bClearInvalid = CachedIdValue.RegistryType.IsValid();
 		CachedIdValue = CurrentId;
-		OnTypeChanged();
+		OnTypeChanged(bClearInvalid);
 	}
 	else if (CurrentId.ItemName != CachedIdValue.ItemName)
 	{
@@ -96,7 +112,7 @@ void SDataRegistryItemNameWidget::OnNameSelected(const FString& NameString)
 	OnSetId.Execute(NewId);
 }
 
-void SDataRegistryItemNameWidget::OnTypeChanged()
+void SDataRegistryItemNameWidget::OnTypeChanged(bool bClearInvalid)
 {
 	int32 SwitcherIndex = 0;
 	CachedIds.Reset();
@@ -132,7 +148,7 @@ void SDataRegistryItemNameWidget::OnTypeChanged()
 		}
 	}
 
-	if (!CachedIdValue.ItemName.IsNone() && !CachedIds.Contains(CachedIdValue))
+	if (!CachedIdValue.ItemName.IsNone() && !CachedIds.Contains(CachedIdValue) && bClearInvalid)
 	{
 		// Name no longer valid, clear
 		OnSetId.Execute(FDataRegistryId(CachedIdValue.RegistryType, NAME_None));
@@ -206,6 +222,15 @@ FText SDataRegistryItemNameWidget::OnGetNameValueText() const
 	return OnGetDisplayText.Execute();
 }
 
+EVisibility SDataRegistryItemNameWidget::GetWarningVisibility() const
+{
+	if (CachedIdValue.RegistryType.IsValid() && !CachedIdValue.ItemName.IsNone() && !CachedIds.Contains(CachedIdValue))
+	{
+		return EVisibility::Visible;
+	}
+	return EVisibility::Collapsed;
+}
+
 TSharedRef<IPropertyTypeCustomization> FDataRegistryIdCustomization::MakeInstance()
 {
 	return MakeShareable(new FDataRegistryIdCustomization);
diff --git a/Engine/Plugins/Runtime/DataRegistry/Source/DataRegistryEditor/Private/DataRegistryIdCustomization.h b/Engine/Plugins/Runtime/DataRegistry/Source/DataRegistryEditor/Private/DataRegistryIdCustomization.h
--- a/Engine/Plugins/Runtime/DataRegistry/Source/DataRegistryEditor/Private/DataRegistryIdCustomization.h
+++ b/Engine/Plugins/Runtime/DataRegistry/Source/DataRegistryEditor/Private/DataRegistryIdCustomization.h
@@ -57,7 +57,7 @@ private:
 	void Tick(const FGeometry& AllottedGeometry, const double InCurrentTime, const float InDeltaTime) override;
 
 	/** Delegate to refresh the drop down when the selected type changes */
-	void OnTypeChanged();
+	void OnTypeChanged(bool bClearInvalid);
 
 	/** Change value and forward to delegates */
 	void OnNameSelected(const FString& NameString);
@@ -75,6 +75,7 @@ private:
 	/** Slate accessors */
 	FString OnGetNameValueString() const;
 	FText OnGetNameValueText() const;
+	EVisibility GetWarningVisibility() const;
 
 	/** This will never be bad while engine is running */
 	class UDataRegistrySubsystem* Subsystem = nullptr;
