Hand merge of CL 18249022, fixing a non-unity build issue originating from CL 18248575

=====

- Fix non-unity build issue with cherry-picked integration from CL 18248575

diff --git a/Engine/Plugins/Runtime/DataRegistry/Source/DataRegistryEditor/Private/DataRegistryIdCustomization.cpp b/Engine/Plugins/Runtime/DataRegistry/Source/DataRegistryEditor/Private/DataRegistryIdCustomization.cpp
--- a/Engine/Plugins/Runtime/DataRegistry/Source/DataRegistryEditor/Private/DataRegistryIdCustomization.cpp
+++ b/Engine/Plugins/Runtime/DataRegistry/Source/DataRegistryEditor/Private/DataRegistryIdCustomization.cpp
@@ -8,6 +8,7 @@
 #include "Widgets/SWidget.h"
 #include "Widgets/Text/STextBlock.h"
 #include "Widgets/Input/SComboButton.h"
+#include "Widgets/Images/SImage.h"
 #include "Components/WidgetSwitcher.h"
 #include "PropertyCustomizationHelpers.h"
 #include "Framework/Application/SlateApplication.h"
