Adding early out optimization to the Data Registry system so that we can skip lookups if we recieve a NAME_None.

diff --git a/Engine/Plugins/Runtime/DataRegistry/Source/DataRegistry/Private/DataRegistry.cpp b/Engine/Plugins/Runtime/DataRegistry/Source/DataRegistry/Private/DataRegistry.cpp
--- a/Engine/Plugins/Runtime/DataRegistry/Source/DataRegistry/Private/DataRegistry.cpp
+++ b/Engine/Plugins/Runtime/DataRegistry/Source/DataRegistry/Private/DataRegistry.cpp
@@ -424,7 +424,7 @@ bool UDataRegistry::ResolveDataRegistryId(FDataRegistryLookup& OutLookup, const
 		}
 		
 		FName ResolvedName = MapIdToResolvedName(ItemId, Source);
-		EDataRegistryAvailability Availability = Source->GetItemAvailability(ResolvedName, PrecachedDataPtr);
+		EDataRegistryAvailability Availability = ResolvedName == NAME_None ? EDataRegistryAvailability::DoesNotExist : Source->GetItemAvailability(ResolvedName, PrecachedDataPtr);
 
 		// If we know it doesn't exist, don't add this to the lookup
 		if (Availability != EDataRegistryAvailability::DoesNotExist)
