#DataRegistry If a data registry or data registry source is disregardforgc, then mark the objects it spawns as part of the root set to avoid violating disregardforgc assumptions

diff --git a/Engine/Plugins/Runtime/DataRegistry/Source/DataRegistry/Private/DataRegistry.cpp b/Engine/Plugins/Runtime/DataRegistry/Source/DataRegistry/Private/DataRegistry.cpp
--- a/Engine/Plugins/Runtime/DataRegistry/Source/DataRegistry/Private/DataRegistry.cpp
+++ b/Engine/Plugins/Runtime/DataRegistry/Source/DataRegistry/Private/DataRegistry.cpp
@@ -921,6 +921,10 @@ void UDataRegistry::RefreshRuntimeSources()
 		{
 			if (IsInitialized() && !Source->IsInitialized())
 			{
+				if (GUObjectArray.IsDisregardForGC(this))
+				{
+					Source->AddToRoot();
+				}
 				Source->Initialize();
 			}
 
@@ -936,6 +940,10 @@ void UDataRegistry::RefreshRuntimeSources()
 		if (OldSource && OldSource->IsInitialized() && !RuntimeSources.Contains(OldSource))
 		{
 			OldSource->Deinitialize();
+			if (GUObjectArray.IsDisregardForGC(this))
+			{
+				OldSource->RemoveFromRoot();
+			}
 		}
 	}
 
diff --git a/Engine/Plugins/Runtime/DataRegistry/Source/DataRegistry/Private/DataRegistrySource.cpp b/Engine/Plugins/Runtime/DataRegistry/Source/DataRegistry/Private/DataRegistrySource.cpp
--- a/Engine/Plugins/Runtime/DataRegistry/Source/DataRegistry/Private/DataRegistrySource.cpp
+++ b/Engine/Plugins/Runtime/DataRegistry/Source/DataRegistry/Private/DataRegistrySource.cpp
@@ -253,6 +253,10 @@ void UMetaDataRegistrySource::RefreshRuntimeSources()
 			UDataRegistrySource* NewSource = CreateTransientSource(GetChildSourceClass());
 			SetDataForChild(SourceName, NewSource);
 			RuntimeChildren.Add(SourceName, NewSource);
+			if (GUObjectArray.IsDisregardForGC(this))
+			{
+				NewSource->AddToRoot();
+			}
 
 			NewSource->Initialize();
 		}
@@ -266,6 +270,10 @@ void UMetaDataRegistrySource::RefreshRuntimeSources()
 			if (MapIt.Value())
 			{
 				MapIt.Value()->Deinitialize();
+				if (GUObjectArray.IsDisregardForGC(this))
+				{
+					MapIt.Value()->RemoveFromRoot();
+				}
 			}
 			MapIt.RemoveCurrent();
 		}
