Redirector support for DataTable row structs in the DataRegistry

#DataRegistry

diff --git a/Engine/Plugins/Runtime/DataRegistry/Source/DataRegistry/Private/DataRegistrySource_DataTable.cpp b/Engine/Plugins/Runtime/DataRegistry/Source/DataRegistry/Private/DataRegistrySource_DataTable.cpp
--- a/Engine/Plugins/Runtime/DataRegistry/Source/DataRegistry/Private/DataRegistrySource_DataTable.cpp
+++ b/Engine/Plugins/Runtime/DataRegistry/Source/DataRegistry/Private/DataRegistrySource_DataTable.cpp
@@ -6,6 +6,7 @@
 #include "DataRegistrySettings.h"
 #include "Interfaces/ITargetPlatform.h"
 #include "Engine/AssetManager.h"
+#include "UObject/CoreRedirects.h"
 #include "UObject/ObjectSaveContext.h"
 
 void UDataRegistrySource_DataTable::SetSourceTable(const TSoftObjectPtr<UDataTable>& InSourceTable, const FDataRegistrySource_DataTableRules& InTableRules)
@@ -368,11 +369,25 @@ bool UMetaDataRegistrySource_DataTable::DoesAssetPassFilter(const FAssetData& As
 		else
 		{
 			// TODO no 100% way to check for inherited row structs, but BP types can't inherit anyway
-			UScriptStruct* RowStruct = FindObject<UScriptStruct>(ANY_PACKAGE, *RowStructureString, true);
+			const UScriptStruct* RowStruct = FindObject<UScriptStruct>(ANY_PACKAGE, *RowStructureString, true);
 
-			if (RowStruct && RowStruct->IsChildOf(ItemStruct))
+			// Check if the row struct is a child of the item struct
+			if (RowStruct != nullptr)
 			{
-				return true;
+				if (RowStruct->IsChildOf(ItemStruct))
+				{
+					return true;
+				}	
+			}
+			// Otherwise check if the row struct has been redirected to the item struct
+			else
+			{
+				FName RowStructureName = FName(RowStructureString);
+				
+				TArray<FCoreRedirectObjectName> PreviousNames;
+				FCoreRedirects::FindPreviousNames(ECoreRedirectFlags::Type_Struct, ItemStruct->GetPathName(), PreviousNames);
+				
+				return PreviousNames.ContainsByPredicate([&RowStructureName](const FCoreRedirectObjectName& PreviousName){ return PreviousName.ObjectName == RowStructureName; });
 			}
 		}
 	}
