#UE Replace another usage of TryFindTypeSlow with FindFirstObject since this code deals with struct names instead of full object paths

diff --git a/Engine/Plugins/Runtime/DataRegistry/Source/DataRegistry/Private/DataRegistrySource_DataTable.cpp b/Engine/Plugins/Runtime/DataRegistry/Source/DataRegistry/Private/DataRegistrySource_DataTable.cpp
--- a/Engine/Plugins/Runtime/DataRegistry/Source/DataRegistry/Private/DataRegistrySource_DataTable.cpp
+++ b/Engine/Plugins/Runtime/DataRegistry/Source/DataRegistry/Private/DataRegistrySource_DataTable.cpp
@@ -369,7 +369,7 @@ bool UMetaDataRegistrySource_DataTable::DoesAssetPassFilter(const FAssetData& As
 		else
 		{
 			// TODO no 100% way to check for inherited row structs, but BP types can't inherit anyway
-			const UScriptStruct* RowStruct = UClass::TryFindTypeSlow<UScriptStruct>(RowStructureString, EFindFirstObjectOptions::ExactClass);
+			const UScriptStruct* RowStruct = FindFirstObject<UScriptStruct>(*RowStructureString, EFindFirstObjectOptions::ExactClass);
 
 			// Check if the row struct is a child of the item struct
 			if (RowStruct != nullptr)
