FNames containing asset paths are deprecated. FSoftObjectPath or FTopLevelAssetPath should be used instead.
Fixups for DataRegistry.

diff --git a/Engine/Plugins/Runtime/DataRegistry/Source/DataRegistry/Private/DataRegistrySource.cpp b/Engine/Plugins/Runtime/DataRegistry/Source/DataRegistry/Private/DataRegistrySource.cpp
--- a/Engine/Plugins/Runtime/DataRegistry/Source/DataRegistry/Private/DataRegistrySource.cpp
+++ b/Engine/Plugins/Runtime/DataRegistry/Source/DataRegistry/Private/DataRegistrySource.cpp
@@ -221,7 +221,7 @@ void UMetaDataRegistrySource::DetermineRuntimeNames(TArray<FName>& OutRuntimeNam
 	{
 		if (DoesAssetPassFilter(AssetData, false))
 		{
-			OutRuntimeNames.Add(AssetData.ObjectPath);
+			OutRuntimeNames.Add(FName(*AssetData.GetObjectPathString()));
 		}
 	}
 }
diff --git a/Engine/Plugins/Runtime/DataRegistry/Source/DataRegistry/Private/DataRegistrySource_CurveTable.cpp b/Engine/Plugins/Runtime/DataRegistry/Source/DataRegistry/Private/DataRegistrySource_CurveTable.cpp
--- a/Engine/Plugins/Runtime/DataRegistry/Source/DataRegistry/Private/DataRegistrySource_CurveTable.cpp
+++ b/Engine/Plugins/Runtime/DataRegistry/Source/DataRegistry/Private/DataRegistrySource_CurveTable.cpp
@@ -331,7 +331,7 @@ bool UMetaDataRegistrySource_CurveTable::SetDataForChild(FName SourceId, UDataRe
 	UDataRegistrySource_CurveTable* ChildCurveTable = Cast<UDataRegistrySource_CurveTable>(ChildSource);
 	if (ensure(ChildCurveTable))
 	{
-		TSoftObjectPtr<UCurveTable> NewTable = TSoftObjectPtr<UCurveTable>(FSoftObjectPath(SourceId));
+		TSoftObjectPtr<UCurveTable> NewTable = TSoftObjectPtr<UCurveTable>(FSoftObjectPath(SourceId.ToString()));
 		ChildCurveTable->SetSourceTable(NewTable, TableRules);
 		return true;
 	}
diff --git a/Engine/Plugins/Runtime/DataRegistry/Source/DataRegistry/Private/DataRegistrySource_DataTable.cpp b/Engine/Plugins/Runtime/DataRegistry/Source/DataRegistry/Private/DataRegistrySource_DataTable.cpp
--- a/Engine/Plugins/Runtime/DataRegistry/Source/DataRegistry/Private/DataRegistrySource_DataTable.cpp
+++ b/Engine/Plugins/Runtime/DataRegistry/Source/DataRegistry/Private/DataRegistrySource_DataTable.cpp
@@ -329,7 +329,7 @@ bool UMetaDataRegistrySource_DataTable::SetDataForChild(FName SourceId, UDataReg
 	UDataRegistrySource_DataTable* ChildDataTable = Cast<UDataRegistrySource_DataTable>(ChildSource);
 	if (ensure(ChildDataTable))
 	{
-		TSoftObjectPtr<UDataTable> NewTable = TSoftObjectPtr<UDataTable>(FSoftObjectPath(SourceId));
+		TSoftObjectPtr<UDataTable> NewTable = TSoftObjectPtr<UDataTable>(FSoftObjectPath(SourceId.ToString()));
 		ChildDataTable->SetSourceTable(NewTable, TableRules);
 		return true;
 	}
