Don't redundantly expand Data Registry virtual scan paths
From PR #9559

diff --git a/Engine/Plugins/Runtime/DataRegistry/Source/DataRegistry/Private/DataRegistrySource.cpp b/Engine/Plugins/Runtime/DataRegistry/Source/DataRegistry/Private/DataRegistrySource.cpp
--- a/Engine/Plugins/Runtime/DataRegistry/Source/DataRegistry/Private/DataRegistrySource.cpp
+++ b/Engine/Plugins/Runtime/DataRegistry/Source/DataRegistry/Private/DataRegistrySource.cpp
@@ -197,7 +197,7 @@ void UMetaDataRegistrySource::DetermineRuntimeNames(TArray<FName>& OutRuntimeNam
 			TempRules.bSkipVirtualPathExpansion = true;
 		}
 
-		AssetManager.SearchAssetRegistryPaths(AssetDataList, SearchRules);
+		AssetManager.SearchAssetRegistryPaths(AssetDataList, TempRules);
 
 		if (bExpandedVirtual)
 		{
