Non-unity compile fix: Add missing include due to removal of EngineTypes.h from AssetManagerTypes.h.
#rnx

diff --git a/Engine/Plugins/Runtime/DataRegistry/Source/DataRegistry/Public/DataRegistry.h b/Engine/Plugins/Runtime/DataRegistry/Source/DataRegistry/Public/DataRegistry.h
--- a/Engine/Plugins/Runtime/DataRegistry/Source/DataRegistry/Public/DataRegistry.h
+++ b/Engine/Plugins/Runtime/DataRegistry/Source/DataRegistry/Public/DataRegistry.h
@@ -3,6 +3,8 @@
 #pragma once
 
 #include "DataRegistryTypes.h"
+#include "Engine/TimerHandle.h"
+
 #include "DataRegistry.generated.h"
 
 class UDataRegistrySource;
diff --git a/Engine/Source/Runtime/Engine/Classes/Engine/AssetManagerTypes.h b/Engine/Source/Runtime/Engine/Classes/Engine/AssetManagerTypes.h
--- a/Engine/Source/Runtime/Engine/Classes/Engine/AssetManagerTypes.h
+++ b/Engine/Source/Runtime/Engine/Classes/Engine/AssetManagerTypes.h
@@ -23,6 +23,7 @@
 #include "AssetRegistry/AssetData.h"
 #include "AssetRegistry/AssetBundleData.h"
 #include "CoreMinimal.h"
+#include "EngineTypes.h"
 #endif
 
 #include "AssetManagerTypes.generated.h"
