Non-unity compile fix: Add missing SubClassOf.h that was previously included by AssetManager.h.
#rnx

diff --git a/Engine/Plugins/Runtime/DataRegistry/Source/DataRegistry/Public/DataRegistrySource.h b/Engine/Plugins/Runtime/DataRegistry/Source/DataRegistry/Public/DataRegistrySource.h
--- a/Engine/Plugins/Runtime/DataRegistry/Source/DataRegistry/Public/DataRegistrySource.h
+++ b/Engine/Plugins/Runtime/DataRegistry/Source/DataRegistry/Public/DataRegistrySource.h
@@ -4,6 +4,8 @@
 
 #include "AssetRegistry/AssetData.h"
 #include "DataRegistry.h"
+#include "Templates/SubclassOf.h"
+
 #include "DataRegistrySource.generated.h"
 
 /** Specifies a source for DataRegistry items, which is an interface that provides asynchronous access to individual structs */
