Removing a reference the DataRegistry doesn't need.

diff --git a/Engine/Plugins/Runtime/DataRegistry/DataRegistry.uplugin b/Engine/Plugins/Runtime/DataRegistry/DataRegistry.uplugin
--- a/Engine/Plugins/Runtime/DataRegistry/DataRegistry.uplugin
+++ b/Engine/Plugins/Runtime/DataRegistry/DataRegistry.uplugin
@@ -31,11 +31,6 @@
 		{
 			"Name": "GameplayTagsEditor",
 			"Enabled": true
-		},
-		{
-			"Name": "AssetManagerEditor",
-			"Enabled": true
 		}
 	]
-
 }
\ No newline at end of file
diff --git a/Engine/Plugins/Runtime/DataRegistry/Source/DataRegistryEditor/DataRegistryEditor.Build.cs b/Engine/Plugins/Runtime/DataRegistry/Source/DataRegistryEditor/DataRegistryEditor.Build.cs
--- a/Engine/Plugins/Runtime/DataRegistry/Source/DataRegistryEditor/DataRegistryEditor.Build.cs
+++ b/Engine/Plugins/Runtime/DataRegistry/Source/DataRegistryEditor/DataRegistryEditor.Build.cs
@@ -19,7 +19,6 @@ namespace UnrealBuildTool.Rules
 					"GameplayTags",
 					"DataRegistry",
 					"GameplayTagsEditor",
-					"AssetManagerEditor",
 					"DataTableEditor",
 					"InputCore",
 					"PropertyEditor",
