Fix data registry sources to never try to precache table refrences during PostLoad as this is not safe in UE5. This would only happen in cases where the asset is missing

diff --git a/Engine/Plugins/Runtime/DataRegistry/Source/DataRegistry/Private/DataRegistrySource_CurveTable.cpp b/Engine/Plugins/Runtime/DataRegistry/Source/DataRegistry/Private/DataRegistrySource_CurveTable.cpp
--- a/Engine/Plugins/Runtime/DataRegistry/Source/DataRegistry/Private/DataRegistrySource_CurveTable.cpp
+++ b/Engine/Plugins/Runtime/DataRegistry/Source/DataRegistry/Private/DataRegistrySource_CurveTable.cpp
@@ -7,6 +7,7 @@
 #include "Engine/AssetManager.h"
 #include "Engine/CurveTable.h"
 #include "UObject/ObjectSaveContext.h"
+#include "UObject/UObjectThreadContext.h"
 #include "Engine/StreamableManager.h"
 
 #include UE_INLINE_GENERATED_CPP_BY_NAME(DataRegistrySource_CurveTable)
@@ -35,12 +36,20 @@ void UDataRegistrySource_CurveTable::SetCachedTable(bool bForceLoad /*= false*/)
 
 	if (!FoundTable && (bForceLoad || TableRules.bPrecacheTable))
 	{
-		if (IsTransientSource() && TableRules.bPrecacheTable && !bForceLoad)
+		if (FUObjectThreadContext::Get().IsRoutingPostLoad)
 		{
-			// Possibly bad sync load, should we warn?
+			// It is not safe to try and preload during a postload, this should only happen if PreloadTable is null
+			UE_LOG(LogDataRegistry, Log, TEXT("Failed to preload table %s for source %s, PreloadTable is %s"), *SourceTable.ToString(), *GetPathName(), *GetNameSafe(PreloadTable));
 		}
+		else 
+		{
+			if (TableRules.bPrecacheTable && !bForceLoad)
+			{
+				UE_LOG(LogDataRegistry, Verbose, TEXT("Synchronously precaching table %s for source %s, PreloadTable is %s"), *SourceTable.ToString(), *GetPathName(), *GetNameSafe(PreloadTable));
+			}
 
-		FoundTable = SourceTable.LoadSynchronous();
+			FoundTable = SourceTable.LoadSynchronous();
+		}
 	}
 
 	if (FoundTable)
diff --git a/Engine/Plugins/Runtime/DataRegistry/Source/DataRegistry/Private/DataRegistrySource_DataTable.cpp b/Engine/Plugins/Runtime/DataRegistry/Source/DataRegistry/Private/DataRegistrySource_DataTable.cpp
--- a/Engine/Plugins/Runtime/DataRegistry/Source/DataRegistry/Private/DataRegistrySource_DataTable.cpp
+++ b/Engine/Plugins/Runtime/DataRegistry/Source/DataRegistry/Private/DataRegistrySource_DataTable.cpp
@@ -10,6 +10,7 @@
 #include "UObject/CoreRedirects.h"
 #include "Engine/StreamableManager.h"
 #include "UObject/ObjectSaveContext.h"
+#include "UObject/UObjectThreadContext.h"
 #include "Misc/PackageName.h"
 
 #include UE_INLINE_GENERATED_CPP_BY_NAME(DataRegistrySource_DataTable)
@@ -38,12 +39,20 @@ void UDataRegistrySource_DataTable::SetCachedTable(bool bForceLoad /*= false*/)
 
 	if (!FoundTable && (bForceLoad || TableRules.bPrecacheTable))
 	{
-		if (IsTransientSource() && TableRules.bPrecacheTable && !bForceLoad)
+		if (FUObjectThreadContext::Get().IsRoutingPostLoad)
 		{
-			// Possibly bad sync load, should we warn?
+			// It is not safe to try and preload during a postload, this should only happen if PreloadTable is null
+			UE_LOG(LogDataRegistry, Log, TEXT("Failed to preload table %s for source %s, PreloadTable is %s"), *SourceTable.ToString(), *GetPathName(), *GetNameSafe(PreloadTable));
 		}
+		else
+		{
+			if (TableRules.bPrecacheTable && !bForceLoad)
+			{
+				UE_LOG(LogDataRegistry, Verbose, TEXT("Synchronously precaching table %s for source %s, PreloadTable is %s"), *SourceTable.ToString(), *GetPathName(), *GetNameSafe(PreloadTable));
+			}
 
-		FoundTable = SourceTable.LoadSynchronous();
+			FoundTable = SourceTable.LoadSynchronous();
+		}
 	}
 
 	if (FoundTable)
