Fixed mismatched copy constructor/assignment operator for FDataRegistryType and FDataRegistryId.

diff --git a/Engine/Plugins/Runtime/DataRegistry/Source/DataRegistry/Public/DataRegistryId.h b/Engine/Plugins/Runtime/DataRegistry/Source/DataRegistry/Public/DataRegistryId.h
--- a/Engine/Plugins/Runtime/DataRegistry/Source/DataRegistry/Public/DataRegistryId.h
+++ b/Engine/Plugins/Runtime/DataRegistry/Source/DataRegistry/Public/DataRegistryId.h
@@ -65,12 +65,6 @@ struct FDataRegistryType
 		return Name != Other.Name;
 	}
 
-	inline FDataRegistryType& operator=(const FDataRegistryType& Other)
-	{
-		Name = Other.Name;
-		return *this;
-	}
-
 	/** Returns true if this is a valid Type */
 	inline bool IsValid() const
 	{
@@ -165,13 +159,6 @@ struct FDataRegistryId
 		return RegistryType != Other.RegistryType || ItemName != Other.ItemName;
 	}
 
-	inline FDataRegistryId& operator=(const FDataRegistryId& Other)
-	{
-		RegistryType = Other.RegistryType;
-		ItemName = Other.ItemName;
-		return *this;
-	}
-
 
 	/** Convert from Type:Name strings, these are inline to avoid linker errors for constructor */
 	inline static FDataRegistryId ParseTypeAndName(const TCHAR* TypeAndName, uint32 Len)
