Fixed mismatched copy constructor/assignment operator in FDataRegistryRequestId.

diff --git a/Engine/Plugins/Runtime/DataRegistry/Source/DataRegistry/Public/DataRegistryTypes.h b/Engine/Plugins/Runtime/DataRegistry/Source/DataRegistry/Public/DataRegistryTypes.h
--- a/Engine/Plugins/Runtime/DataRegistry/Source/DataRegistry/Public/DataRegistryTypes.h
+++ b/Engine/Plugins/Runtime/DataRegistry/Source/DataRegistry/Public/DataRegistryTypes.h
@@ -433,12 +433,6 @@ struct FDataRegistryRequestId
 		return RequestId < Other.RequestId;
 	}
 
-	FDataRegistryRequestId& operator=(const FDataRegistryRequestId& Other)
-	{
-		RequestId = Other.RequestId;
-		return *this;
-	}
-
 	friend inline uint32 GetTypeHash(const FDataRegistryRequestId& Key)
 	{
 		return GetTypeHash(Key.RequestId);
