Add support for objects referenced by a struct when collecting references from a data table or data registry. FInstancedStruct depends on this to prevent GC of any objects it references.

diff --git a/Engine/Plugins/Runtime/DataRegistry/Source/DataRegistry/Private/DataRegistry.cpp b/Engine/Plugins/Runtime/DataRegistry/Source/DataRegistry/Private/DataRegistry.cpp
--- a/Engine/Plugins/Runtime/DataRegistry/Source/DataRegistry/Private/DataRegistry.cpp
+++ b/Engine/Plugins/Runtime/DataRegistry/Source/DataRegistry/Private/DataRegistry.cpp
@@ -1129,7 +1129,7 @@ void UDataRegistry::AddReferencedObjects(UObject* InThis, FReferenceCollector& C
 
 			if (Item && Item->ItemMemory)
 			{
-				Collector.AddPropertyReferences(ItemStruct, Item->ItemMemory, This);
+				Collector.AddPropertyReferencesWithStructARO(ItemStruct, Item->ItemMemory, This);
 			}
 		}
 	}
diff --git a/Engine/Source/Runtime/Engine/Private/DataTable.cpp b/Engine/Source/Runtime/Engine/Private/DataTable.cpp
--- a/Engine/Source/Runtime/Engine/Private/DataTable.cpp
+++ b/Engine/Source/Runtime/Engine/Private/DataTable.cpp
@@ -278,7 +278,7 @@ void UDataTable::AddReferencedObjects(UObject* InThis, FReferenceCollector& Coll
 		{
 			if (uint8* RowData = Pair.Value)
 			{
-				Collector.AddPropertyReferences(This->RowStruct, RowData, This);
+				Collector.AddPropertyReferencesWithStructARO(This->RowStruct, RowData, This);
 			}
 		}
 	}
