Added a DataRegistryResolver interface that allows us to specify temporary local objects instead of having to allocate a shared pointer for every single one. Deprecated the resolve function that returns the resolver as TSharedPtr since it may be the temporary one now.

diff --git a/Engine/Plugins/Runtime/DataRegistry/Source/DataRegistry/Private/DataRegistry.cpp b/Engine/Plugins/Runtime/DataRegistry/Source/DataRegistry/Private/DataRegistry.cpp
--- a/Engine/Plugins/Runtime/DataRegistry/Source/DataRegistry/Private/DataRegistry.cpp
+++ b/Engine/Plugins/Runtime/DataRegistry/Source/DataRegistry/Private/DataRegistry.cpp
@@ -988,8 +988,8 @@ FName UDataRegistry::MapIdToResolvedName(const FDataRegistryId& ItemId, const UD
 {
 	// Try resolver stack before falling back to raw name
 	FName ResolvedName;
-	TSharedPtr<FDataRegistryResolver> FoundResolver = FDataRegistryResolverScope::ResolveIdToName(ResolvedName, ItemId, this, RegistrySource);
-	if (FoundResolver.IsValid())
+	FDataRegistryResolver* FoundResolver = FDataRegistryResolverScope::ResolveNameFromId(ResolvedName, ItemId, this, RegistrySource);
+	if (FoundResolver)
 	{
 		return ResolvedName;
 	}
diff --git a/Engine/Plugins/Runtime/DataRegistry/Source/DataRegistry/Private/DataRegistryTypes.cpp b/Engine/Plugins/Runtime/DataRegistry/Source/DataRegistry/Private/DataRegistryTypes.cpp
--- a/Engine/Plugins/Runtime/DataRegistry/Source/DataRegistry/Private/DataRegistryTypes.cpp
+++ b/Engine/Plugins/Runtime/DataRegistry/Source/DataRegistry/Private/DataRegistryTypes.cpp
@@ -218,46 +218,83 @@ void FDataRegistryCache::ClearCache(bool bClearRequests)
 	}
 }
 
-TArray<TSharedPtr<FDataRegistryResolver> > FDataRegistryResolverScope::ResolverStack;
+TArray<FDataRegistryResolverScope::FResolverStackEntry> FDataRegistryResolverScope::ResolverStack;
+
+FDataRegistryResolverScope::FDataRegistryResolverScope(FDataRegistryResolver& ScopeResolver)
+	: AddedScopeResolver(&ScopeResolver)
+{
+	// Only valid to call on game thread
+	check(IsInGameThread());
+	ResolverStack.Emplace(ScopeResolver);
+}
 
 FDataRegistryResolverScope::FDataRegistryResolverScope(const TSharedPtr<FDataRegistryResolver>& ScopeResolver)
+	: AddedScopeResolver(ScopeResolver.Get())
 {
 	// Only valid to call on game thread
 	check(IsInGameThread() && ScopeResolver.IsValid());
-	
-	ResolverStack.Push(ScopeResolver);
-	StackAtAdd = ResolverStack.Num();
+	ResolverStack.Emplace(ScopeResolver);
 }
 
 FDataRegistryResolverScope::~FDataRegistryResolverScope()
 {
-	if (ensureMsgf(StackAtAdd == ResolverStack.Num(), TEXT("FDataRegistryResolverScope has invalid stack depth at destroy! Unsafe to copy scope or add a global resolve while scope is active")))
+	if (ensureMsgf(!ResolverStack.IsEmpty() && (AddedScopeResolver == ResolverStack.Last().AsRawPtr), TEXT("FDataRegistryResolverScope has invalid stack depth at destroy! Unsafe to copy scope or add a global resolve while scope is active")))
 	{
 		ResolverStack.Pop();
 	}
+	else
+	{
+		// Don't want to leak a raw pointer, so manually find and remove it here
+		for (int32 i = ResolverStack.Num() - 1; i >= 0; --i)
+		{
+			if (ResolverStack[i].AsRawPtr == AddedScopeResolver)
+			{
+				ResolverStack.RemoveAt(i);
+				break;
+			}
+		}
+	}
 }
 
 void FDataRegistryResolverScope::RegisterGlobalResolver(const TSharedPtr<FDataRegistryResolver>& ScopeResolver)
 {
 	// Only valid to call on game thread
 	check(IsInGameThread() && ScopeResolver.IsValid());
-	
 	ResolverStack.Add(ScopeResolver);
 }
 
 void FDataRegistryResolverScope::UnregisterGlobalResolver(const TSharedPtr<FDataRegistryResolver>& ScopeResolver)
 {
-	int32 RemovedIndex = ResolverStack.Remove(ScopeResolver);
-	ensureMsgf(RemovedIndex != INDEX_NONE, TEXT("FDataRegistryResolverScope::UnregisterGlobalResolver called with resolver that was never registered!"));
+	const int32 RemovedIndex = ResolverStack.FindLastByPredicate([ScopeResolver](const FResolverStackEntry& Entry) { return Entry.AsSharedPtr == ScopeResolver; });
+	if (ensureMsgf(RemovedIndex != INDEX_NONE, TEXT("FDataRegistryResolverScope::UnregisterGlobalResolver called with resolver that was never registered!")))
+	{
+		ResolverStack.RemoveAt(RemovedIndex);
+	}
 }
 
 TSharedPtr<FDataRegistryResolver> FDataRegistryResolverScope::ResolveIdToName(FName& OutResolvedName, const FDataRegistryId& ItemId, const class UDataRegistry* Registry, const class UDataRegistrySource* RegistrySource)
+{
+	// This version will only crawl shared pointer resolvers since we don't want it to resolve successfully but then return null
+	for (int32 i = ResolverStack.Num() - 1; i >= 0; i--)
+	{
+		FResolverStackEntry& Entry = ResolverStack[i];
+		if (Entry.AsSharedPtr.IsValid() && Entry.AsSharedPtr->ResolveIdToName(OutResolvedName, ItemId, Registry, RegistrySource))
+		{
+			return ResolverStack[i].AsSharedPtr;
+		}
+	}
+
+	return nullptr;
+}
+
+FDataRegistryResolver* FDataRegistryResolverScope::ResolveNameFromId(FName& OutResolvedName, const FDataRegistryId& ItemId, const class UDataRegistry* Registry, const class UDataRegistrySource* RegistrySource)
 {
 	for (int32 i = ResolverStack.Num() - 1; i >= 0; i--)
 	{
-		if (ResolverStack[i]->ResolveIdToName(OutResolvedName, ItemId, Registry, RegistrySource))
+		FResolverStackEntry& Entry = ResolverStack[i];
+		if (Entry.AsRawPtr->ResolveIdToName(OutResolvedName, ItemId, Registry, RegistrySource))
 		{
-			return ResolverStack[i];
+			return ResolverStack[i].AsRawPtr;
 		}
 	}
 
@@ -268,7 +305,7 @@ bool FDataRegistryResolverScope::IsStackVolatile()
 {
 	for (int32 i = ResolverStack.Num() - 1; i >= 0; i--)
 	{
-		if (ResolverStack[i]->IsVolatile())
+		if (ResolverStack[i].AsRawPtr->IsVolatile())
 		{
 			return true;
 		}
diff --git a/Engine/Plugins/Runtime/DataRegistry/Source/DataRegistry/Public/DataRegistryTypes.h b/Engine/Plugins/Runtime/DataRegistry/Source/DataRegistry/Public/DataRegistryTypes.h
--- a/Engine/Plugins/Runtime/DataRegistry/Source/DataRegistry/Public/DataRegistryTypes.h
+++ b/Engine/Plugins/Runtime/DataRegistry/Source/DataRegistry/Public/DataRegistryTypes.h
@@ -458,6 +458,9 @@ struct FDataRegistryResolver
 /** Scope object to set a temporary resolver, or register a global one */
 struct DATAREGISTRY_API FDataRegistryResolverScope
 {
+	/** Creates a temporary resolver scope without requiring a dynamically allocated object. Only valid on game thread */
+	FDataRegistryResolverScope(FDataRegistryResolver& ScopeResolver);
+
 	/** Creates a temporary resolver scope, will use the passed in resolver until scope returns. Only valid on game thread */
 	FDataRegistryResolverScope(const TSharedPtr<FDataRegistryResolver>& ScopeResolver);
 
@@ -473,15 +476,36 @@ struct DATAREGISTRY_API FDataRegistryResolverScope
 	/** Use the stack to resolve an ID, will return the resolver used if found */
 	static TSharedPtr<FDataRegistryResolver> ResolveIdToName(FName& OutResolvedName, const FDataRegistryId& ItemId, const class UDataRegistry* Registry, const class UDataRegistrySource* RegistrySource);
 
+	/** Use the stack to resolve an ID, will return the resolver as raw pointer if found */
+	static FDataRegistryResolver* ResolveNameFromId(FName& OutResolvedName, const FDataRegistryId& ItemId, const class UDataRegistry* Registry, const class UDataRegistrySource* RegistrySource);
+
 	/** Returns true if there are any volatile resolvers on the stack */
 	static bool IsStackVolatile();
 
 private:
 	// Stack depth at point this was added
-	int32 StackAtAdd = 0;
+	FDataRegistryResolver* AddedScopeResolver;
+
+	// Global and temporary resolver stack, will go from last added back to 0, optionally able to push stack objects
+	struct FResolverStackEntry
+	{
+		FResolverStackEntry(const TSharedPtr<FDataRegistryResolver>& ScopeResolver)
+			: AsSharedPtr(ScopeResolver)
+			, AsRawPtr(ScopeResolver.Get())
+		{
+		}
+
+		FResolverStackEntry(FDataRegistryResolver& ScopeResolver)
+			: AsSharedPtr()
+			, AsRawPtr(&ScopeResolver)
+		{
+		}
+
+		TSharedPtr<FDataRegistryResolver> AsSharedPtr;
+		FDataRegistryResolver* AsRawPtr;
+	};
 
-	// Global and temporary resolver stack, will go from last added back to 0
-	static TArray<TSharedPtr<FDataRegistryResolver> > ResolverStack;
+	static TArray<FResolverStackEntry> ResolverStack;
 
 };
 
