Deprecated ResolveIdToName in order to favour the raw-pointer variant.

diff --git a/Engine/Plugins/Runtime/DataRegistry/Source/DataRegistry/Public/DataRegistryTypes.h b/Engine/Plugins/Runtime/DataRegistry/Source/DataRegistry/Public/DataRegistryTypes.h
--- a/Engine/Plugins/Runtime/DataRegistry/Source/DataRegistry/Public/DataRegistryTypes.h
+++ b/Engine/Plugins/Runtime/DataRegistry/Source/DataRegistry/Public/DataRegistryTypes.h
@@ -474,6 +474,7 @@ struct DATAREGISTRY_API FDataRegistryResolverScope
 	static void UnregisterGlobalResolver(const TSharedPtr<FDataRegistryResolver>& ScopeResolver);
 
 	/** Use the stack to resolve an ID, will return the resolver used if found */
+	UE_DEPRECATED(5.3, "Use ResolveNameFromId instead to allow support temporary scope resolvers by raw pointer.")
 	static TSharedPtr<FDataRegistryResolver> ResolveIdToName(FName& OutResolvedName, const FDataRegistryId& ItemId, const class UDataRegistry* Registry, const class UDataRegistrySource* RegistrySource);
 
 	/** Use the stack to resolve an ID, will return the resolver as raw pointer if found */
