Fix issue where transient DataRegistry data sources were not being correctly removed from the RuntimeChildren array when a registry was Deinitialized
Move the GC logic for transient sources into Create/ResetTransientSource so it always matches the RuntimeChildren array

diff --git a/Engine/Plugins/Runtime/DataRegistry/Source/DataRegistry/Private/DataRegistry.cpp b/Engine/Plugins/Runtime/DataRegistry/Source/DataRegistry/Private/DataRegistry.cpp
--- a/Engine/Plugins/Runtime/DataRegistry/Source/DataRegistry/Private/DataRegistry.cpp
+++ b/Engine/Plugins/Runtime/DataRegistry/Source/DataRegistry/Private/DataRegistry.cpp
@@ -923,10 +923,6 @@ void UDataRegistry::RefreshRuntimeSources()
 		{
 			if (IsInitialized() && !Source->IsInitialized())
 			{
-				if (GUObjectArray.IsDisregardForGC(this))
-				{
-					Source->AddToRoot();
-				}
 				Source->Initialize();
 			}
 
@@ -935,16 +931,18 @@ void UDataRegistry::RefreshRuntimeSources()
 		}
 	}
 
-	// Uninitialize any orphaned sources
+	// Deinitialize any orphaned sources
 	for (int32 i = 0; i < OldRuntimeSources.Num(); i++)
 	{
 		UDataRegistrySource* OldSource = OldRuntimeSources[i];
 		if (OldSource && OldSource->IsInitialized() && !RuntimeSources.Contains(OldSource))
 		{
 			OldSource->Deinitialize();
-			if (GUObjectArray.IsDisregardForGC(this))
+
+			if (OldSource->IsTransientSource())
 			{
-				OldSource->RemoveFromRoot();
+				// Transient sources should have already been removed during RefreshRuntimeSources
+				UE_LOG(LogDataRegistry, Warning, TEXT("RefreshRuntimeSources found orphaned transient source %s for registry %s!"), *OldSource->GetPathName(), *RegistryType.ToString());
 			}
 		}
 	}
diff --git a/Engine/Plugins/Runtime/DataRegistry/Source/DataRegistry/Private/DataRegistrySource.cpp b/Engine/Plugins/Runtime/DataRegistry/Source/DataRegistry/Private/DataRegistrySource.cpp
--- a/Engine/Plugins/Runtime/DataRegistry/Source/DataRegistry/Private/DataRegistrySource.cpp
+++ b/Engine/Plugins/Runtime/DataRegistry/Source/DataRegistry/Private/DataRegistrySource.cpp
@@ -98,6 +98,19 @@ bool UDataRegistrySource::IsTransientSource() const
 	return HasAnyFlags(RF_Transient);
 }
 
+bool UDataRegistrySource::ResetTransientSource()
+{
+	if (ensure(IsTransientSource() && !IsInitialized()))
+	{
+		// Mark that this is no longer part of a parent source, the calling function must remove it from any parent-child arrays
+		ParentSource = nullptr;
+		RemoveFromRoot();
+		return true;
+	}
+
+	return false;
+}
+
 UDataRegistrySource* UDataRegistrySource::GetOriginalSource()
 {
 	return ParentSource ? ToRawPtr(ParentSource) : this;
@@ -141,6 +154,12 @@ UDataRegistrySource* UDataRegistrySource::CreateTransientSource(TSubclassOf<UDat
 
 	UDataRegistrySource* ChildSource = NewObject<UDataRegistrySource>(GetRegistry(), *SourceClass, NAME_None, RF_Transient);
 	ChildSource->ParentSource = this;
+	
+	if (GUObjectArray.IsDisregardForGC(this))
+	{
+		ChildSource->AddToRoot();
+	}
+
 	return ChildSource;
 }
 
@@ -237,34 +256,32 @@ void UMetaDataRegistrySource::RefreshRuntimeSources()
 {
 	if (!IsInitialized())
 	{
-		return;
+		// This will clear out all active runtime children
+		RuntimeNames.Reset();
 	}
-
-	TArray<FName> OldRuntimeNames = RuntimeNames;
-	DetermineRuntimeNames(RuntimeNames);
-
-	for (FName SourceName : RuntimeNames)
+	else
 	{
-		TObjectPtr<UDataRegistrySource>* FoundSource = RuntimeChildren.Find(SourceName);
-		if (FoundSource && *FoundSource)
-		{
-			SetDataForChild(SourceName, *FoundSource);
-		}
-		else
+		DetermineRuntimeNames(RuntimeNames);
+
+		for (FName SourceName : RuntimeNames)
 		{
-			UDataRegistrySource* NewSource = CreateTransientSource(GetChildSourceClass());
-			SetDataForChild(SourceName, NewSource);
-			RuntimeChildren.Add(SourceName, NewSource);
-			if (GUObjectArray.IsDisregardForGC(this))
+			TObjectPtr<UDataRegistrySource>* FoundSource = RuntimeChildren.Find(SourceName);
+			if (FoundSource && *FoundSource)
 			{
-				NewSource->AddToRoot();
+				SetDataForChild(SourceName, *FoundSource);
+				ensureMsgf((*FoundSource)->IsInitialized(), TEXT("Runtime source %s was not initialized after refresh!"), *FoundSource->GetPathName());
+			}
+			else
+			{
+				UDataRegistrySource* NewSource = CreateTransientSource(GetChildSourceClass());
+				SetDataForChild(SourceName, NewSource);
+				RuntimeChildren.Add(SourceName, NewSource);
+				NewSource->Initialize();
 			}
-
-			NewSource->Initialize();
 		}
 	}
 
-	// Clear out old undesired children, this will get deinitialized later
+	// Clear out old undesired children and deinitialize them
 	for (auto MapIt = RuntimeChildren.CreateIterator(); MapIt; ++MapIt)
 	{
 		if (!RuntimeNames.Contains(MapIt.Key()))
@@ -272,10 +289,7 @@ void UMetaDataRegistrySource::RefreshRuntimeSources()
 			if (MapIt.Value())
 			{
 				MapIt.Value()->Deinitialize();
-				if (GUObjectArray.IsDisregardForGC(this))
-				{
-					MapIt.Value()->RemoveFromRoot();
-				}
+				MapIt.Value()->ResetTransientSource();
 			}
 			MapIt.RemoveCurrent();
 		}
diff --git a/Engine/Plugins/Runtime/DataRegistry/Source/DataRegistry/Public/DataRegistrySource.h b/Engine/Plugins/Runtime/DataRegistry/Source/DataRegistry/Public/DataRegistrySource.h
--- a/Engine/Plugins/Runtime/DataRegistry/Source/DataRegistry/Public/DataRegistrySource.h
+++ b/Engine/Plugins/Runtime/DataRegistry/Source/DataRegistry/Public/DataRegistrySource.h
@@ -69,6 +69,9 @@ public:
 	/** Returns true if this is a runtime-only source */
 	virtual bool IsTransientSource() const;
 
+	/** Resets transient source so it can be correctly garbage collected */
+	virtual bool ResetTransientSource();
+
 	/** Returns the editor-defined source, which is either this or the parent source */
 	virtual UDataRegistrySource* GetOriginalSource();
 
