Added a OnSubsystemInitialized callback to UDataRegistrySubsystem, notifying when all known data registries have been scanned and initialized.

diff --git a/Engine/Plugins/Runtime/DataRegistry/Source/DataRegistry/Private/DataRegistrySubsystem.cpp b/Engine/Plugins/Runtime/DataRegistry/Source/DataRegistry/Private/DataRegistrySubsystem.cpp
--- a/Engine/Plugins/Runtime/DataRegistry/Source/DataRegistry/Private/DataRegistrySubsystem.cpp
+++ b/Engine/Plugins/Runtime/DataRegistry/Source/DataRegistry/Private/DataRegistrySubsystem.cpp
@@ -370,6 +370,11 @@ bool UDataRegistrySubsystem::IsConfigEnabled(bool bWarnIfNotEnabled /*= false*/)
 	return true;
 }
 
+FDataRegistrySubsystemInitializedCallback& UDataRegistrySubsystem::OnSubsystemInitialized()
+{
+	return OnSubsystemInitializedCallback;
+}
+
 void UDataRegistrySubsystem::InitializeAllRegistries(bool bResetIfInitialized)
 {
 	for (const FRegistryMapPair& RegistryPair : RegistryMap)
@@ -386,6 +391,8 @@ void UDataRegistrySubsystem::InitializeAllRegistries(bool bResetIfInitialized)
 	}
 
 	bFullyInitialized = true;
+
+	OnSubsystemInitialized().Broadcast();
 }
 
 void UDataRegistrySubsystem::DeinitializeAllRegistries()
diff --git a/Engine/Plugins/Runtime/DataRegistry/Source/DataRegistry/Public/DataRegistrySubsystem.h b/Engine/Plugins/Runtime/DataRegistry/Source/DataRegistry/Public/DataRegistrySubsystem.h
--- a/Engine/Plugins/Runtime/DataRegistry/Source/DataRegistry/Public/DataRegistrySubsystem.h
+++ b/Engine/Plugins/Runtime/DataRegistry/Source/DataRegistry/Public/DataRegistrySubsystem.h
@@ -147,6 +147,9 @@ public:
 	/** Returns true if the system is enabled via any config scan settings, will optionally warn if not enabled */
 	bool IsConfigEnabled(bool bWarnIfNotEnabled = false) const;
 
+	/** Accesses the delegate called when the subsystem has finished scanning for and initializing all known data registries */
+	FDataRegistrySubsystemInitializedCallback& OnSubsystemInitialized();
+
 	/** Initializes all loaded registries and prepares them for queries */
 	void InitializeAllRegistries(bool bResetIfInitialized = false);
 
@@ -251,6 +254,9 @@ protected:
 	// True if initialization is ready to start, will be true even if config disabled
 	bool bReadyForInitialization = false;
 
+	/** Callback for when the subsystem has finished scanning for and initializing all known data registries */
+	FDataRegistrySubsystemInitializedCallback OnSubsystemInitializedCallback;
+
 #if WITH_EDITOR
 	virtual void PreBeginPIE(bool bStartSimulate);
 	virtual void EndPIE(bool bStartSimulate);
diff --git a/Engine/Plugins/Runtime/DataRegistry/Source/DataRegistry/Public/DataRegistryTypes.h b/Engine/Plugins/Runtime/DataRegistry/Source/DataRegistry/Public/DataRegistryTypes.h
--- a/Engine/Plugins/Runtime/DataRegistry/Source/DataRegistry/Public/DataRegistryTypes.h
+++ b/Engine/Plugins/Runtime/DataRegistry/Source/DataRegistry/Public/DataRegistryTypes.h
@@ -523,6 +523,9 @@ DECLARE_DELEGATE_OneParam(FDataRegistryBatchAcquireCallback, EDataRegistryAcquir
 /** Multicast delegate broadcast called when a data registry's cache version has changed */
 DECLARE_MULTICAST_DELEGATE_OneParam(FDataRegistryCacheVersionCallback, class UDataRegistry*);
 
+/** Multicast delegate broadcast called when the data registry subsystem has finished scanning for and initializing all known data registries */
+DECLARE_MULTICAST_DELEGATE(FDataRegistrySubsystemInitializedCallback);
+
 
 DECLARE_LOG_CATEGORY_EXTERN(LogDataRegistry, Log, All);
 
