Fix DataRegistry issue where an invalid table source would stop the cache from working for other data sources, it now warns and returns DoesNotExist if a table load fails
Add editor validation logic to GameFeatureAction_DataRegistrySource to try and prevent this from happening

diff --git a/Engine/Plugins/Experimental/GameFeatures/Source/GameFeatures/Private/GameFeatureAction_DataRegistrySource.cpp b/Engine/Plugins/Experimental/GameFeatures/Source/GameFeatures/Private/GameFeatureAction_DataRegistrySource.cpp
--- a/Engine/Plugins/Experimental/GameFeatures/Source/GameFeatures/Private/GameFeatureAction_DataRegistrySource.cpp
+++ b/Engine/Plugins/Experimental/GameFeatures/Source/GameFeatures/Private/GameFeatureAction_DataRegistrySource.cpp
@@ -10,6 +10,8 @@
 
 #if WITH_EDITOR
 #include "Misc/DataValidation.h"
+#include "AssetRegistry/AssetData.h"
+#include "AssetRegistry/AssetRegistryModule.h"
 #endif
 
 #include UE_INLINE_GENERATED_CPP_BY_NAME(GameFeatureAction_DataRegistrySource)
@@ -242,6 +244,30 @@ EDataValidationResult UGameFeatureAction_DataRegistrySource::IsDataValid(FDataVa
 			Result = EDataValidationResult::Invalid;
 		}
 
+		if (!Entry.CurveTableToAdd.IsNull())
+		{
+			FAssetData TableAssetData = IAssetRegistry::Get()->GetAssetByObjectPath(Entry.CurveTableToAdd.ToSoftObjectPath());
+
+			// This will catch normal curve tables, composites, and any reasonably named subclass without doing a slow load
+			if (!TableAssetData.IsValid() || !TableAssetData.AssetClassPath.GetAssetName().ToString().Contains(TEXT("CurveTable")))
+			{
+				Context.AddError(FText::Format(LOCTEXT("DataRegistrySourceMissingCurveTable", "Path {0} does not point to valid curvetable at index {1} in SourcesToAdd"), FText::FromString(Entry.CurveTableToAdd.ToString()), FText::AsNumber(EntryIndex)));
+				Result = EDataValidationResult::Invalid;
+			}
+		}
+
+		if (!Entry.DataTableToAdd.IsNull())
+		{
+			FAssetData TableAssetData = IAssetRegistry::Get()->GetAssetByObjectPath(Entry.DataTableToAdd.ToSoftObjectPath());
+
+			// This will catch normal data tables, composites, and any reasonably named subclass without doing a slow load
+			if (!TableAssetData.IsValid() || !TableAssetData.AssetClassPath.GetAssetName().ToString().Contains(TEXT("DataTable")))
+			{
+				Context.AddError(FText::Format(LOCTEXT("DataRegistrySourceMissingDataTable", "Path {0} does not point to valid datatable at index {1} in SourcesToAdd"), FText::FromString(Entry.DataTableToAdd.ToString()), FText::AsNumber(EntryIndex)));
+				Result = EDataValidationResult::Invalid;
+			}
+		}
+
 		if (Entry.bServerSource == false && Entry.bClientSource == false)
 		{
 			Context.AddError(FText::Format(LOCTEXT("DataRegistrySourceNeverUsed", "Source not specified to load on either client or server, it will be unused at index {0} in SourcesToAdd"), FText::AsNumber(EntryIndex)));
diff --git a/Engine/Plugins/Runtime/DataRegistry/Source/DataRegistry/Private/DataRegistrySource_CurveTable.cpp b/Engine/Plugins/Runtime/DataRegistry/Source/DataRegistry/Private/DataRegistrySource_CurveTable.cpp
--- a/Engine/Plugins/Runtime/DataRegistry/Source/DataRegistry/Private/DataRegistrySource_CurveTable.cpp
+++ b/Engine/Plugins/Runtime/DataRegistry/Source/DataRegistry/Private/DataRegistrySource_CurveTable.cpp
@@ -49,6 +49,14 @@ void UDataRegistrySource_CurveTable::SetCachedTable(bool bForceLoad /*= false*/)
 			}
 
 			FoundTable = SourceTable.LoadSynchronous();
+			if (!FoundTable)
+			{
+				if (!SourceTable.IsNull())
+				{
+					UE_LOG(LogDataRegistry, Warning, TEXT("Force loading table %s for source %s failed! Source data is invalid and will be ignored."), *SourceTable.ToString(), *GetPathName());
+				}
+				bInvalidSourceTable = true;
+			}
 		}
 	}
 
@@ -59,22 +67,27 @@ void UDataRegistrySource_CurveTable::SetCachedTable(bool bForceLoad /*= false*/)
 		if (FoundTable->HasAnyFlags(RF_NeedLoad))
 		{
 			UE_LOG(LogDataRegistry, Error, TEXT("Cannot initialize DataRegistry source %s, Preload table was not set, resave in editor!"), *GetPathName());
+			bInvalidSourceTable = true;
 		}
 		else if (!ItemStruct || !ItemStruct->IsChildOf(FRealCurve::StaticStruct()))
 		{
 			UE_LOG(LogDataRegistry, Error, TEXT("Cannot initialize DataRegistry source %s, Curve tables only work with curve items!"), *GetPathName(), *FoundTable->GetPathName());
+			bInvalidSourceTable = true;
 		}
 		else if (ItemStruct->IsChildOf(FSimpleCurve::StaticStruct()) && FoundTable->GetCurveTableMode() != ECurveTableMode::SimpleCurves)
 		{
 			UE_LOG(LogDataRegistry, Error, TEXT("Cannot initialize DataRegistry source %s, only simple curve tables are supported for SimpleCurve items, reimport with linear interpolation type or select different table"), *GetPathName(), *FoundTable->GetPathName());
+			bInvalidSourceTable = true;
 		}
 		else if (ItemStruct->IsChildOf(FRichCurve::StaticStruct()) && FoundTable->GetCurveTableMode() != ECurveTableMode::RichCurves)
 		{
 			UE_LOG(LogDataRegistry, Error, TEXT("Cannot initialize DataRegistry source %s, only rich curve tables are supported for RichCurve items, reimport with cubic interpolation type or select different table"), *GetPathName(), *FoundTable->GetPathName());
+			bInvalidSourceTable = true;
 		}
 		else
 		{
 			CachedTable = FoundTable;
+			bInvalidSourceTable = false;
 
 #if WITH_EDITOR
 			if (GIsEditor)
@@ -150,10 +163,12 @@ EDataRegistryAvailability UDataRegistrySource_CurveTable::GetItemAvailability(co
 			return EDataRegistryAvailability::DoesNotExist;
 		}
 	}
-	else
+	else if (bInvalidSourceTable)
 	{
-		return EDataRegistryAvailability::Unknown;
+		return EDataRegistryAvailability::DoesNotExist;
 	}
+
+	return EDataRegistryAvailability::Unknown;
 }
 
 void UDataRegistrySource_CurveTable::GetResolvedNames(TArray<FName>& Names) const
@@ -297,6 +312,13 @@ void UDataRegistrySource_CurveTable::OnTableLoaded()
 
 	SetCachedTable(false);
 
+	// If we failed to set a cached table, the source data is invalid
+	if (CachedTable == nullptr)
+	{
+		UE_LOG(LogDataRegistry, Warning, TEXT("Loading table %s for source %s failed! Source data is invalid and will be ignored."), *SourceTable.ToString(), *GetPathName());
+		bInvalidSourceTable = true;
+	}
+
 	HandlePendingAcquires();
 }
 
@@ -307,13 +329,6 @@ void UDataRegistrySource_CurveTable::EditorRefreshSource()
 	SetCachedTable(false);
 }
 
-void UDataRegistrySource_CurveTable::PreSave(const ITargetPlatform* TargetPlatform)
-{
-	PRAGMA_DISABLE_DEPRECATION_WARNINGS;
-	Super::PreSave(TargetPlatform);
-	PRAGMA_ENABLE_DEPRECATION_WARNINGS;
-}
-
 void UDataRegistrySource_CurveTable::PreSave(FObjectPreSaveContext ObjectSaveContext)
 {
 	Super::PreSave(ObjectSaveContext);
diff --git a/Engine/Plugins/Runtime/DataRegistry/Source/DataRegistry/Private/DataRegistrySource_DataTable.cpp b/Engine/Plugins/Runtime/DataRegistry/Source/DataRegistry/Private/DataRegistrySource_DataTable.cpp
--- a/Engine/Plugins/Runtime/DataRegistry/Source/DataRegistry/Private/DataRegistrySource_DataTable.cpp
+++ b/Engine/Plugins/Runtime/DataRegistry/Source/DataRegistry/Private/DataRegistrySource_DataTable.cpp
@@ -52,6 +52,14 @@ void UDataRegistrySource_DataTable::SetCachedTable(bool bForceLoad /*= false*/)
 			}
 
 			FoundTable = SourceTable.LoadSynchronous();
+			if (!FoundTable)
+			{
+				if (!SourceTable.IsNull())
+				{
+					UE_LOG(LogDataRegistry, Warning, TEXT("Force loading table %s for source %s failed! Source data is invalid and will be ignored."), *SourceTable.ToString(), *GetPathName());
+				}
+				bInvalidSourceTable = true;
+			}
 		}
 	}
 
@@ -63,18 +71,22 @@ void UDataRegistrySource_DataTable::SetCachedTable(bool bForceLoad /*= false*/)
 		if (FoundTable->HasAnyFlags(RF_NeedLoad))
 		{
 			UE_LOG(LogDataRegistry, Error, TEXT("Cannot initialize DataRegistry source %s, Preload table was not set, resave in editor!"), *GetPathName());
+			bInvalidSourceTable = true;
 		}
 		else if(!ItemStruct || !RowStruct)
 		{
 			UE_LOG(LogDataRegistry, Error, TEXT("Cannot initialize DataRegistry source %s, Table %s or registry is invalid!"), *GetPathName(), *FoundTable->GetPathName());
+			bInvalidSourceTable = true;
 		}
 		else if (!RowStruct->IsChildOf(ItemStruct))
 		{
 			UE_LOG(LogDataRegistry, Error, TEXT("Cannot initialize DataRegistry source %s, Table %s type does not match %s"), *GetPathName(), *FoundTable->GetPathName(), *RowStruct->GetName(), *ItemStruct->GetName());
+			bInvalidSourceTable = true;
 		}
 		else
 		{
 			CachedTable = FoundTable;
+			bInvalidSourceTable = false;
 
 #if WITH_EDITOR
 			if (GIsEditor)
@@ -150,10 +162,12 @@ EDataRegistryAvailability UDataRegistrySource_DataTable::GetItemAvailability(con
 			return EDataRegistryAvailability::DoesNotExist;
 		}
 	}
-	else
+	else if (bInvalidSourceTable)
 	{
-		return EDataRegistryAvailability::Unknown;
+		return EDataRegistryAvailability::DoesNotExist;
 	}
+
+	return EDataRegistryAvailability::Unknown;
 }
 
 void UDataRegistrySource_DataTable::GetResolvedNames(TArray<FName>& Names) const
@@ -297,6 +311,13 @@ void UDataRegistrySource_DataTable::OnTableLoaded()
 
 	SetCachedTable(false);
 
+	// If we failed to set a cached table, the source data is invalid
+	if (CachedTable == nullptr)
+	{
+		UE_LOG(LogDataRegistry, Warning, TEXT("Loading table %s for source %s failed! Source data is invalid and will be ignored."), *SourceTable.ToString(), *GetPathName());
+		bInvalidSourceTable = true;
+	}
+
 	HandlePendingAcquires();
 }
 
@@ -307,13 +328,6 @@ void UDataRegistrySource_DataTable::EditorRefreshSource()
 	SetCachedTable(false);
 }
 
-void UDataRegistrySource_DataTable::PreSave(const ITargetPlatform* TargetPlatform)
-{
-	PRAGMA_DISABLE_DEPRECATION_WARNINGS;
-	Super::PreSave(TargetPlatform);
-	PRAGMA_ENABLE_DEPRECATION_WARNINGS;
-}
-
 void UDataRegistrySource_DataTable::PreSave(FObjectPreSaveContext ObjectSaveContext)
 {
 	Super::PreSave(ObjectSaveContext);
diff --git a/Engine/Plugins/Runtime/DataRegistry/Source/DataRegistry/Public/DataRegistrySource_CurveTable.h b/Engine/Plugins/Runtime/DataRegistry/Source/DataRegistry/Public/DataRegistrySource_CurveTable.h
--- a/Engine/Plugins/Runtime/DataRegistry/Source/DataRegistry/Public/DataRegistrySource_CurveTable.h
+++ b/Engine/Plugins/Runtime/DataRegistry/Source/DataRegistry/Public/DataRegistrySource_CurveTable.h
@@ -38,6 +38,9 @@ protected:
 	/** Last time this was accessed */
 	mutable float LastAccessTime;
 
+	/** This is set if the table fails to load or could never load */
+	bool bInvalidSourceTable = false;
+
 	/** Handle for in progress load */
 	TSharedPtr<FStreamableHandle> LoadingTableHandle;
 
@@ -72,10 +75,6 @@ protected:
 
 #if WITH_EDITOR
 	virtual void EditorRefreshSource();
-	PRAGMA_DISABLE_DEPRECATION_WARNINGS // Suppress compiler warning on override of deprecated function
-	UE_DEPRECATED(5.0, "Use version that takes FObjectPreSaveContext instead.")
-	virtual void PreSave(const class ITargetPlatform* TargetPlatform) override;
-	PRAGMA_ENABLE_DEPRECATION_WARNINGS
 	virtual void PreSave(FObjectPreSaveContext ObjectSaveContext) override;
 
 #endif
diff --git a/Engine/Plugins/Runtime/DataRegistry/Source/DataRegistry/Public/DataRegistrySource_DataTable.h b/Engine/Plugins/Runtime/DataRegistry/Source/DataRegistry/Public/DataRegistrySource_DataTable.h
--- a/Engine/Plugins/Runtime/DataRegistry/Source/DataRegistry/Public/DataRegistrySource_DataTable.h
+++ b/Engine/Plugins/Runtime/DataRegistry/Source/DataRegistry/Public/DataRegistrySource_DataTable.h
@@ -55,6 +55,9 @@ protected:
 	/** Last time this was accessed */
 	mutable float LastAccessTime;
 
+	/** This is set if the table fails to load or could never load */
+	bool bInvalidSourceTable = false;
+
 	/** Handle for in progress load */
 	TSharedPtr<FStreamableHandle> LoadingTableHandle;
 
@@ -89,10 +92,6 @@ protected:
 
 #if WITH_EDITOR
 	virtual void EditorRefreshSource();
-	PRAGMA_DISABLE_DEPRECATION_WARNINGS // Suppress compiler warning on override of deprecated function
-	UE_DEPRECATED(5.0, "Use version that takes FObjectPreSaveContext instead.")
-	virtual void PreSave(const class ITargetPlatform* TargetPlatform) override;
-	PRAGMA_ENABLE_DEPRECATION_WARNINGS
 	virtual void PreSave(FObjectPreSaveContext ObjectSaveContext) override;
 
 #endif
