Fix crash when deleting a data registry source in the editor, it was trying to deinitialize a partially destroyed object in the transient package

diff --git a/Engine/Plugins/Runtime/DataRegistry/Source/DataRegistry/Private/DataRegistry.cpp b/Engine/Plugins/Runtime/DataRegistry/Source/DataRegistry/Private/DataRegistry.cpp
--- a/Engine/Plugins/Runtime/DataRegistry/Source/DataRegistry/Private/DataRegistry.cpp
+++ b/Engine/Plugins/Runtime/DataRegistry/Source/DataRegistry/Private/DataRegistry.cpp
@@ -935,7 +935,7 @@ void UDataRegistry::RefreshRuntimeSources()
 	for (int32 i = 0; i < OldRuntimeSources.Num(); i++)
 	{
 		UDataRegistrySource* OldSource = OldRuntimeSources[i];
-		if (OldSource && OldSource->IsInitialized() && !RuntimeSources.Contains(OldSource))
+		if (IsValid(OldSource) && OldSource->GetOuter() == this && OldSource->IsInitialized() && !RuntimeSources.Contains(OldSource))
 		{
 			OldSource->Deinitialize();
 
