Adding an Enum return value to Register Specific Asset, because if the asset is already registered then we should not log out a warning, which causes Tools to trigger warnings.

diff --git a/Engine/Plugins/Runtime/DataRegistry/Source/DataRegistry/Private/DataRegistry.cpp b/Engine/Plugins/Runtime/DataRegistry/Source/DataRegistry/Private/DataRegistry.cpp
--- a/Engine/Plugins/Runtime/DataRegistry/Source/DataRegistry/Private/DataRegistry.cpp
+++ b/Engine/Plugins/Runtime/DataRegistry/Source/DataRegistry/Private/DataRegistry.cpp
@@ -193,26 +193,37 @@ void UDataRegistry::RuntimeRefreshIfNeeded()
 	}
 }
 
-bool UDataRegistry::RegisterSpecificAsset(const FAssetData& AssetData, int32 AssetPriority /*= 0*/)
+EDataRegistryRegisterAssetResult UDataRegistry::RegisterSpecificAsset(const FAssetData& AssetData, int32 AssetPriority /*= 0*/)
 {
-	bool bMadeChange = false;
+	EDataRegistryRegisterAssetResult RegisterAssetResult = EDataRegistryRegisterAssetResult::NotRegistered;
 
 	for (int32 i = 0; i < DataSources.Num(); i++)
 	{
 		UDataRegistrySource* Source = DataSources[i];
-		if (Source && Source->RegisterSpecificAsset(AssetData, AssetPriority))
+		if (Source)
 		{
-			bMadeChange = true;
+			EDataRegistryRegisterAssetResult CurrentRegisterAssetResult = Source->RegisterSpecificAsset(AssetData, AssetPriority);
+
+			// Don't clobber Asset already registered if Not Registered, Only move to a higher Result state, with success being the highest.
+			if (CurrentRegisterAssetResult > RegisterAssetResult)
+			{
+				RegisterAssetResult = CurrentRegisterAssetResult;
+			}
 		}
 	}
 
-	if (bMadeChange && IsInitialized())
+	if (IsRegisterAssetResultSuccess(RegisterAssetResult) && IsInitialized())
 	{
 		// Don't want to do a full reset, but do clear cache as lookup rules may have changed
 		RefreshRuntimeSources();
 	}
 	
-	return bMadeChange;
+	return RegisterAssetResult;
+}
+
+bool UDataRegistry::IsRegisterAssetResultSuccess(EDataRegistryRegisterAssetResult RegisterAssetStatus)
+{
+	return RegisterAssetStatus == EDataRegistryRegisterAssetResult::RegisteredSuccesfully;
 }
 
 bool UDataRegistry::UnregisterSpecificAsset(const FSoftObjectPath& AssetPath)
diff --git a/Engine/Plugins/Runtime/DataRegistry/Source/DataRegistry/Private/DataRegistrySource.cpp b/Engine/Plugins/Runtime/DataRegistry/Source/DataRegistry/Private/DataRegistrySource.cpp
--- a/Engine/Plugins/Runtime/DataRegistry/Source/DataRegistry/Private/DataRegistrySource.cpp
+++ b/Engine/Plugins/Runtime/DataRegistry/Source/DataRegistry/Private/DataRegistrySource.cpp
@@ -121,9 +121,9 @@ bool UDataRegistrySource::IsSpecificAssetRegistered(const FSoftObjectPath& Asset
 	return false;
 }
 
-bool UDataRegistrySource::RegisterSpecificAsset(const FAssetData& AssetData, int32 AssetPriority)
+EDataRegistryRegisterAssetResult UDataRegistrySource::RegisterSpecificAsset(const FAssetData& AssetData, int32 AssetPriority)
 {
-	return false;
+	return EDataRegistryRegisterAssetResult::NotRegistered;
 }
 
 bool UDataRegistrySource::UnregisterSpecificAsset(const FSoftObjectPath& AssetPath)
@@ -325,7 +325,7 @@ bool UMetaDataRegistrySource::IsSpecificAssetRegistered(const FSoftObjectPath& A
 	return false;
 }
 
-bool UMetaDataRegistrySource::RegisterSpecificAsset(const FAssetData& AssetData, int32 AssetPriority)
+EDataRegistryRegisterAssetResult UMetaDataRegistrySource::RegisterSpecificAsset(const FAssetData& AssetData, int32 AssetPriority)
 {
 	bool bMadeChange = false;
 	
@@ -337,7 +337,7 @@ bool UMetaDataRegistrySource::RegisterSpecificAsset(const FAssetData& AssetData,
 			if (Existing.Value == AssetPriority)
 			{
 				// Nothing to do
-				return false;
+				return EDataRegistryRegisterAssetResult::AssetAlreadyRegistered;
 			}
 			bMadeChange = true;
 			Existing.Value = AssetPriority;
@@ -359,7 +359,7 @@ bool UMetaDataRegistrySource::RegisterSpecificAsset(const FAssetData& AssetData,
 		SortRegisteredAssets();
 	}
 
-	return bMadeChange;
+	return bMadeChange ? EDataRegistryRegisterAssetResult::RegisteredSuccesfully : EDataRegistryRegisterAssetResult::NotRegistered;
 }
 
 bool UMetaDataRegistrySource::UnregisterSpecificAsset(const FSoftObjectPath& AssetPath)
diff --git a/Engine/Plugins/Runtime/DataRegistry/Source/DataRegistry/Private/DataRegistrySubsystem.cpp b/Engine/Plugins/Runtime/DataRegistry/Source/DataRegistry/Private/DataRegistrySubsystem.cpp
--- a/Engine/Plugins/Runtime/DataRegistry/Source/DataRegistry/Private/DataRegistrySubsystem.cpp
+++ b/Engine/Plugins/Runtime/DataRegistry/Source/DataRegistry/Private/DataRegistrySubsystem.cpp
@@ -498,7 +498,7 @@ bool UDataRegistrySubsystem::RegisterSpecificAsset(FDataRegistryType RegistryTyp
 			UDataRegistry* Registry = RegistryPair.Value;
 			if (Registry)
 			{
-				bMadeChange |= Registry->RegisterSpecificAsset(AssetData, AssetPriority);
+				bMadeChange |= Registry->IsRegisterAssetResultSuccess(Registry->RegisterSpecificAsset(AssetData, AssetPriority));
 			}
 		}
 		return bMadeChange;
@@ -507,7 +507,7 @@ bool UDataRegistrySubsystem::RegisterSpecificAsset(FDataRegistryType RegistryTyp
 	UDataRegistry* FoundRegistry = GetRegistryForType(RegistryType);
 	if (FoundRegistry)
 	{
-		return FoundRegistry->RegisterSpecificAsset(AssetData, AssetPriority);
+		return FoundRegistry->IsRegisterAssetResultSuccess(FoundRegistry->RegisterSpecificAsset(AssetData, AssetPriority));
 	}
 
 	return false;
@@ -672,25 +672,30 @@ void UDataRegistrySubsystem::ApplyPreregisterMap(UDataRegistry* Registry)
 		const FTopLevelAssetPath ObjectClassPath(TEXT("/Script/CoreUObject"), TEXT("Object"));
 		for (int32 i = 0; i < FoundPreregister->Num(); i++)
 		{
+			EDataRegistryRegisterAssetResult RegisterAssetResult = EDataRegistryRegisterAssetResult::NotRegistered;
 			bool bRegistered = false;
 			const FSoftObjectPath& AssetPath = (*FoundPreregister)[i].Key;
 			FAssetData AssetData;
 			if (AssetManager.GetAssetDataForPath(AssetPath, AssetData))
 			{
-				bRegistered = Registry->RegisterSpecificAsset(AssetData, (*FoundPreregister)[i].Value);
+				RegisterAssetResult = Registry->RegisterSpecificAsset(AssetData, (*FoundPreregister)[i].Value);
 			}
 			else if (Settings->CanIgnoreMissingAssetData())
 			{
 				// Construct fake asset data and register that
 				AssetData = FAssetData(AssetPath.GetLongPackageName(), AssetPath.GetAssetPathString(), ObjectClassPath);
-				bRegistered = Registry->RegisterSpecificAsset(AssetData, (*FoundPreregister)[i].Value);
+				RegisterAssetResult = Registry->RegisterSpecificAsset(AssetData, (*FoundPreregister)[i].Value);
 			}
 
-			if (!bRegistered)
+			if (RegisterAssetResult == EDataRegistryRegisterAssetResult::NotRegistered)
 			{
-				// If specific type is mentioned, it is expected to always succeed
+				// If specific type is mentioned, it is expected to always succeed unless it was already registered
 				UE_LOG(LogDataRegistry, Warning, TEXT("ApplyPreregisterMap failed to register %s with %s, there needs to be a meta source that handles registered assets with matching data"), *AssetPath.ToString(), *Registry->GetRegistryType().ToString());
 			}
+			else if (RegisterAssetResult == EDataRegistryRegisterAssetResult::AssetAlreadyRegistered)
+			{
+				UE_LOG(LogDataRegistry, Log, TEXT("ApplyPreregisterMap did not register %s with %s, this asset was already registered with the data registry."), *AssetPath.ToString(), *Registry->GetRegistryType().ToString());
+			}
 		}
 	}
 
diff --git a/Engine/Plugins/Runtime/DataRegistry/Source/DataRegistry/Public/DataRegistry.h b/Engine/Plugins/Runtime/DataRegistry/Source/DataRegistry/Public/DataRegistry.h
--- a/Engine/Plugins/Runtime/DataRegistry/Source/DataRegistry/Public/DataRegistry.h
+++ b/Engine/Plugins/Runtime/DataRegistry/Source/DataRegistry/Public/DataRegistry.h
@@ -61,7 +61,10 @@ public:
 	virtual void RuntimeRefreshIfNeeded();
 
 	/** Attempt to register a specified asset with a source, returns true if any changes were made. Can be used to update priority for existing asset as well */
-	virtual bool RegisterSpecificAsset(const FAssetData& AssetData, int32 AssetPriority = 0);
+	virtual EDataRegistryRegisterAssetResult RegisterSpecificAsset(const FAssetData& AssetData, int32 AssetPriority = 0);
+
+	/** Use this to confirm that a DataRegistryRegisterAssetResult is a success or not. Virtual so that different registries can decide what counts as a success.*/
+	virtual bool IsRegisterAssetResultSuccess(EDataRegistryRegisterAssetResult RegisterAssetStatus);
 
 	/** Removes references to a specific asset, returns bool if it was removed */
 	virtual bool UnregisterSpecificAsset(const FSoftObjectPath& AssetPath);
diff --git a/Engine/Plugins/Runtime/DataRegistry/Source/DataRegistry/Public/DataRegistrySource.h b/Engine/Plugins/Runtime/DataRegistry/Source/DataRegistry/Public/DataRegistrySource.h
--- a/Engine/Plugins/Runtime/DataRegistry/Source/DataRegistry/Public/DataRegistrySource.h
+++ b/Engine/Plugins/Runtime/DataRegistry/Source/DataRegistry/Public/DataRegistrySource.h
@@ -79,7 +79,7 @@ public:
 	virtual bool IsSpecificAssetRegistered(const FSoftObjectPath& AssetPath) const;
 
 	/** Attempt to register a specified asset with a source, returns true if any changes were made. Can be used to update priority for existing asset as well */
-	virtual bool RegisterSpecificAsset(const FAssetData& AssetData, int32 AssetPriority = 0);
+	virtual EDataRegistryRegisterAssetResult RegisterSpecificAsset(const FAssetData& AssetData, int32 AssetPriority = 0);
 
 	/** Removes references to a specific asset, returns bool if it was removed */
 	virtual bool UnregisterSpecificAsset(const FSoftObjectPath& AssetPath);
@@ -148,7 +148,7 @@ public:
 	virtual void RefreshRuntimeSources() override;
 	virtual void AddRuntimeSources(TArray<UDataRegistrySource*>& OutRuntimeSources) override;
 	virtual bool IsSpecificAssetRegistered(const FSoftObjectPath& AssetPath) const override;
-	virtual bool RegisterSpecificAsset(const FAssetData& AssetData, int32 AssetPriority) override;
+	virtual EDataRegistryRegisterAssetResult RegisterSpecificAsset(const FAssetData& AssetData, int32 AssetPriority) override;
 	virtual bool UnregisterSpecificAsset(const FSoftObjectPath& AssetPath) override;
 	virtual int32 UnregisterAssetsWithPriority(int32 AssetPriority) override;
 
diff --git a/Engine/Plugins/Runtime/DataRegistry/Source/DataRegistry/Public/DataRegistryTypes.h b/Engine/Plugins/Runtime/DataRegistry/Source/DataRegistry/Public/DataRegistryTypes.h
--- a/Engine/Plugins/Runtime/DataRegistry/Source/DataRegistry/Public/DataRegistryTypes.h
+++ b/Engine/Plugins/Runtime/DataRegistry/Source/DataRegistry/Public/DataRegistryTypes.h
@@ -213,6 +213,19 @@ enum class EDataRegistryCacheGetStatus : uint8
 	FoundPersistent,
 };
 
+/** Error code returned when attempting to register a specific asset. */
+enum class EDataRegistryRegisterAssetResult: uint8
+{
+	/** Asset was not registered */
+	NotRegistered,
+
+	/** Asset already registered in the data registry*/
+	AssetAlreadyRegistered,
+
+	/** Asset was successfully registered */
+	RegisteredSuccesfully, // Maintain as highest result.
+};
+
 /** Where the cache version comes from, the upper values of this can be used for game-specific sources */
 enum class EDataRegistryCacheVersionSource : uint8
 {
