Adjusting Data Registry Register Asset Result changes to be more informative for how to upgrade to the new Engine version.

diff --git a/Engine/Plugins/Runtime/DataRegistry/Source/DataRegistry/Private/DataRegistry.cpp b/Engine/Plugins/Runtime/DataRegistry/Source/DataRegistry/Private/DataRegistry.cpp
--- a/Engine/Plugins/Runtime/DataRegistry/Source/DataRegistry/Private/DataRegistry.cpp
+++ b/Engine/Plugins/Runtime/DataRegistry/Source/DataRegistry/Private/DataRegistry.cpp
@@ -212,7 +212,7 @@ EDataRegistryRegisterAssetResult UDataRegistry::RegisterSpecificAsset(const FAss
 		}
 	}
 
-	if (IsRegisterAssetResultSuccess(RegisterAssetResult) && IsInitialized())
+	if (DidRegisterAssetResultCauseChange(RegisterAssetResult) && IsInitialized())
 	{
 		// Don't want to do a full reset, but do clear cache as lookup rules may have changed
 		RefreshRuntimeSources();
@@ -221,7 +221,7 @@ EDataRegistryRegisterAssetResult UDataRegistry::RegisterSpecificAsset(const FAss
 	return RegisterAssetResult;
 }
 
-bool UDataRegistry::IsRegisterAssetResultSuccess(EDataRegistryRegisterAssetResult RegisterAssetStatus)
+bool UDataRegistry::DidRegisterAssetResultCauseChange(EDataRegistryRegisterAssetResult RegisterAssetStatus) const
 {
 	return RegisterAssetStatus == EDataRegistryRegisterAssetResult::RegisteredSuccesfully;
 }
diff --git a/Engine/Plugins/Runtime/DataRegistry/Source/DataRegistry/Private/DataRegistrySubsystem.cpp b/Engine/Plugins/Runtime/DataRegistry/Source/DataRegistry/Private/DataRegistrySubsystem.cpp
--- a/Engine/Plugins/Runtime/DataRegistry/Source/DataRegistry/Private/DataRegistrySubsystem.cpp
+++ b/Engine/Plugins/Runtime/DataRegistry/Source/DataRegistry/Private/DataRegistrySubsystem.cpp
@@ -498,7 +498,7 @@ bool UDataRegistrySubsystem::RegisterSpecificAsset(FDataRegistryType RegistryTyp
 			UDataRegistry* Registry = RegistryPair.Value;
 			if (Registry)
 			{
-				bMadeChange |= Registry->IsRegisterAssetResultSuccess(Registry->RegisterSpecificAsset(AssetData, AssetPriority));
+				bMadeChange |= Registry->DidRegisterAssetResultCauseChange(Registry->RegisterSpecificAsset(AssetData, AssetPriority));
 			}
 		}
 		return bMadeChange;
@@ -507,7 +507,7 @@ bool UDataRegistrySubsystem::RegisterSpecificAsset(FDataRegistryType RegistryTyp
 	UDataRegistry* FoundRegistry = GetRegistryForType(RegistryType);
 	if (FoundRegistry)
 	{
-		return FoundRegistry->IsRegisterAssetResultSuccess(FoundRegistry->RegisterSpecificAsset(AssetData, AssetPriority));
+		return FoundRegistry->DidRegisterAssetResultCauseChange(FoundRegistry->RegisterSpecificAsset(AssetData, AssetPriority));
 	}
 
 	return false;
diff --git a/Engine/Plugins/Runtime/DataRegistry/Source/DataRegistry/Public/DataRegistry.h b/Engine/Plugins/Runtime/DataRegistry/Source/DataRegistry/Public/DataRegistry.h
--- a/Engine/Plugins/Runtime/DataRegistry/Source/DataRegistry/Public/DataRegistry.h
+++ b/Engine/Plugins/Runtime/DataRegistry/Source/DataRegistry/Public/DataRegistry.h
@@ -60,11 +60,14 @@ public:
 	/** Conditionally refresh the runtime state if needed */
 	virtual void RuntimeRefreshIfNeeded();
 
-	/** Attempt to register a specified asset with a source, returns true if any changes were made. Can be used to update priority for existing asset as well */
+	/** 
+	 * Attempt to register a specified asset with a source, returns an EDataRegistryRegisterAssetResult, indicating whether the asset was registered or not. Can be used to update priority for existing asset as well 
+	 * @return An EDataRegistryRegisterAssetResult that indicates if the asset Was Registered, whether it failed to register, or did not register because the asset was already registered.
+	 */
 	virtual EDataRegistryRegisterAssetResult RegisterSpecificAsset(const FAssetData& AssetData, int32 AssetPriority = 0);
 
-	/** Use this to confirm that a DataRegistryRegisterAssetResult is a success or not. Virtual so that different registries can decide what counts as a success.*/
-	virtual bool IsRegisterAssetResultSuccess(EDataRegistryRegisterAssetResult RegisterAssetStatus);
+	/** Use this to confirm that a DataRegistryRegisterAssetResult caused a change to the Data Registry. Virtual so that different registries can decide what may cause a change.*/
+	virtual bool DidRegisterAssetResultCauseChange(EDataRegistryRegisterAssetResult RegisterAssetStatus) const;
 
 	/** Removes references to a specific asset, returns bool if it was removed */
 	virtual bool UnregisterSpecificAsset(const FSoftObjectPath& AssetPath);
diff --git a/Engine/Plugins/Runtime/DataRegistry/Source/DataRegistry/Public/DataRegistrySource.h b/Engine/Plugins/Runtime/DataRegistry/Source/DataRegistry/Public/DataRegistrySource.h
--- a/Engine/Plugins/Runtime/DataRegistry/Source/DataRegistry/Public/DataRegistrySource.h
+++ b/Engine/Plugins/Runtime/DataRegistry/Source/DataRegistry/Public/DataRegistrySource.h
@@ -78,7 +78,10 @@ public:
 	/** Returns true if this asset is already registered with this source */
 	virtual bool IsSpecificAssetRegistered(const FSoftObjectPath& AssetPath) const;
 
-	/** Attempt to register a specified asset with a source, returns true if any changes were made. Can be used to update priority for existing asset as well */
+	/**
+	 * Attempt to register a specified asset with a source, returns an EDataRegistryRegisterAssetResult, indicating whether the asset was registered or not. Can be used to update priority for existing asset as well
+	 * @return An EDataRegistryRegisterAssetResult that indicates if the asset Was Registered, whether it failed to register, or did not register because the asset was already registered.
+	 */
 	virtual EDataRegistryRegisterAssetResult RegisterSpecificAsset(const FAssetData& AssetData, int32 AssetPriority = 0);
 
 	/** Removes references to a specific asset, returns bool if it was removed */
