Invalidate data registry cache when source data or curve table changes

diff --git a/Engine/Plugins/Runtime/DataRegistry/Source/DataRegistry/Private/DataRegistrySource_CurveTable.cpp b/Engine/Plugins/Runtime/DataRegistry/Source/DataRegistry/Private/DataRegistrySource_CurveTable.cpp
--- a/Engine/Plugins/Runtime/DataRegistry/Source/DataRegistry/Private/DataRegistrySource_CurveTable.cpp
+++ b/Engine/Plugins/Runtime/DataRegistry/Source/DataRegistry/Private/DataRegistrySource_CurveTable.cpp
@@ -24,12 +24,10 @@ void UDataRegistrySource_CurveTable::SetSourceTable(const TSoftObjectPtr<UCurveT
 
 void UDataRegistrySource_CurveTable::SetCachedTable(bool bForceLoad /*= false*/)
 {
-#if WITH_EDITOR
-	if (CachedTable && GIsEditor)
+	if (CachedTable)
 	{
 		CachedTable->OnCurveTableChanged().RemoveAll(this);
 	}
-#endif
 
 	CachedTable = nullptr;
 	UCurveTable* FoundTable = SourceTable.Get();
@@ -89,13 +87,8 @@ void UDataRegistrySource_CurveTable::SetCachedTable(bool bForceLoad /*= false*/)
 			CachedTable = FoundTable;
 			bInvalidSourceTable = false;
 
-#if WITH_EDITOR
-			if (GIsEditor)
-			{
-				// Listen for changes like row 
-				CachedTable->OnCurveTableChanged().AddUObject(this, &UDataRegistrySource_CurveTable::EditorRefreshSource);
-			}
-#endif
+			// Listen for changes like row 
+			CachedTable->OnCurveTableChanged().AddUObject(this, &UDataRegistrySource_CurveTable::OnDataTableChanged);
 		}
 	}
 
@@ -322,13 +315,23 @@ void UDataRegistrySource_CurveTable::OnTableLoaded()
 	HandlePendingAcquires();
 }
 
+void UDataRegistrySource_CurveTable::OnDataTableChanged()
+{
 #if WITH_EDITOR
+	if (GIsEditor)
+	{
+		SetCachedTable(false);
+	}
+#endif
 
-void UDataRegistrySource_CurveTable::EditorRefreshSource()
-{
-	SetCachedTable(false);
+	if (IsInitialized())
+	{
+		GetRegistry()->InvalidateCacheVersion();
+	}
 }
 
+#if WITH_EDITOR
+
 void UDataRegistrySource_CurveTable::PreSave(FObjectPreSaveContext ObjectSaveContext)
 {
 	Super::PreSave(ObjectSaveContext);
diff --git a/Engine/Plugins/Runtime/DataRegistry/Source/DataRegistry/Private/DataRegistrySource_DataTable.cpp b/Engine/Plugins/Runtime/DataRegistry/Source/DataRegistry/Private/DataRegistrySource_DataTable.cpp
--- a/Engine/Plugins/Runtime/DataRegistry/Source/DataRegistry/Private/DataRegistrySource_DataTable.cpp
+++ b/Engine/Plugins/Runtime/DataRegistry/Source/DataRegistry/Private/DataRegistrySource_DataTable.cpp
@@ -27,12 +27,10 @@ void UDataRegistrySource_DataTable::SetSourceTable(const TSoftObjectPtr<UDataTab
 
 void UDataRegistrySource_DataTable::SetCachedTable(bool bForceLoad /*= false*/)
 {
-#if WITH_EDITOR
-	if (CachedTable && GIsEditor)
+	if (CachedTable)
 	{
 		CachedTable->OnDataTableChanged().RemoveAll(this);
 	}
-#endif
 
 	CachedTable = nullptr;
 	UDataTable* FoundTable = SourceTable.Get();
@@ -88,13 +86,8 @@ void UDataRegistrySource_DataTable::SetCachedTable(bool bForceLoad /*= false*/)
 			CachedTable = FoundTable;
 			bInvalidSourceTable = false;
 
-#if WITH_EDITOR
-			if (GIsEditor)
-			{
-				// Listen for changes like row 
-				CachedTable->OnDataTableChanged().AddUObject(this, &UDataRegistrySource_DataTable::EditorRefreshSource);
-			}
-#endif
+			// Listen for changes like row 
+			CachedTable->OnDataTableChanged().AddUObject(this, &UDataRegistrySource_DataTable::OnDataTableChanged);
 		}
 	}
 
@@ -321,13 +314,23 @@ void UDataRegistrySource_DataTable::OnTableLoaded()
 	HandlePendingAcquires();
 }
 
+void UDataRegistrySource_DataTable::OnDataTableChanged()
+{
 #if WITH_EDITOR
+	if (GIsEditor)
+	{
+		SetCachedTable(false);
+	}
+#endif
 
-void UDataRegistrySource_DataTable::EditorRefreshSource()
-{
-	SetCachedTable(false);
+	if (IsInitialized())
+	{
+		GetRegistry()->InvalidateCacheVersion();
+	}
 }
 
+#if WITH_EDITOR
+
 void UDataRegistrySource_DataTable::PreSave(FObjectPreSaveContext ObjectSaveContext)
 {
 	Super::PreSave(ObjectSaveContext);
diff --git a/Engine/Plugins/Runtime/DataRegistry/Source/DataRegistry/Public/DataRegistrySource_CurveTable.h b/Engine/Plugins/Runtime/DataRegistry/Source/DataRegistry/Public/DataRegistrySource_CurveTable.h
--- a/Engine/Plugins/Runtime/DataRegistry/Source/DataRegistry/Public/DataRegistrySource_CurveTable.h
+++ b/Engine/Plugins/Runtime/DataRegistry/Source/DataRegistry/Public/DataRegistrySource_CurveTable.h
@@ -73,10 +73,10 @@ protected:
 	// Object interface
 	virtual void PostLoad() override;
 
+	virtual void OnDataTableChanged();
+
 #if WITH_EDITOR
-	virtual void EditorRefreshSource();
 	virtual void PreSave(FObjectPreSaveContext ObjectSaveContext) override;
-
 #endif
 };
 
diff --git a/Engine/Plugins/Runtime/DataRegistry/Source/DataRegistry/Public/DataRegistrySource_DataTable.h b/Engine/Plugins/Runtime/DataRegistry/Source/DataRegistry/Public/DataRegistrySource_DataTable.h
--- a/Engine/Plugins/Runtime/DataRegistry/Source/DataRegistry/Public/DataRegistrySource_DataTable.h
+++ b/Engine/Plugins/Runtime/DataRegistry/Source/DataRegistry/Public/DataRegistrySource_DataTable.h
@@ -90,10 +90,10 @@ protected:
 	// Object interface
 	virtual void PostLoad() override;
 
+	virtual void OnDataTableChanged();
+
 #if WITH_EDITOR
-	virtual void EditorRefreshSource();
 	virtual void PreSave(FObjectPreSaveContext ObjectSaveContext) override;
-
 #endif
 };
 
