SoftDataRegistryOrTable: Struct is used to hook up a Data Table or a Data Registry.  It has a common interface for accessing the entries within it.  SoftObjectPtrs to DataTables will autoconvert existing data if the property name matches.

diff --git a/Engine/Plugins/Runtime/DataRegistry/Source/DataRegistry/Private/DataRegistry.cpp b/Engine/Plugins/Runtime/DataRegistry/Source/DataRegistry/Private/DataRegistry.cpp
--- a/Engine/Plugins/Runtime/DataRegistry/Source/DataRegistry/Private/DataRegistry.cpp
+++ b/Engine/Plugins/Runtime/DataRegistry/Source/DataRegistry/Private/DataRegistry.cpp
@@ -810,6 +810,25 @@ UDataRegistrySource* UDataRegistry::LookupSource(FName& OutResolvedName, const F
 	return nullptr;
 }
 
+void UDataRegistry::GetItemNames(TArray<FName>& ItemNames) const
+{
+	TMap<FDataRegistryId, const uint8*> CachedItemMap;
+	const UScriptStruct* ItemTypeStruct = nullptr;
+
+	const FDataRegistryCacheGetResult Result = GetAllCachedItems(CachedItemMap, ItemTypeStruct);
+	if (!Result.WasFound())
+	{
+		UE_LOG(LogDataRegistry, Warning, TEXT("[%hs] No Registry Data found  Registry:%s"), __FUNCTION__, *GetPathName());
+		return;
+	}
+
+	for (TPair<FDataRegistryId, const uint8*>& CachedItem : CachedItemMap)
+	{
+		const FDataRegistryId& RegistryId = CachedItem.Key;
+		ItemNames.Add(RegistryId.ItemName);
+	}
+}
+
 FDataRegistryCacheGetResult UDataRegistry::GetAllCachedItems(TMap<FDataRegistryId, const uint8*>& OutItemMap, const UScriptStruct*& OutItemStruct) const
 {
 	TSet<FDataRegistryId> IdSet;
diff --git a/Engine/Plugins/Runtime/DataRegistry/Source/DataRegistry/Private/SoftDataRegistryOrTable.cpp b/Engine/Plugins/Runtime/DataRegistry/Source/DataRegistry/Private/SoftDataRegistryOrTable.cpp
new file mode 100644
--- /dev/null
+++ b/Engine/Plugins/Runtime/DataRegistry/Source/DataRegistry/Private/SoftDataRegistryOrTable.cpp
@@ -0,0 +1,133 @@
+// Copyright Epic Games, Inc. All Rights Reserved.
+
+#include "SoftDataRegistryOrTable.h"
+
+#include "DataRegistrySubsystem.h"
+
+#include UE_INLINE_GENERATED_CPP_BY_NAME(SoftDataRegistryOrTable)
+
+bool FSoftDataRegistryOrTable::Serialize(FArchive& Ar)
+{
+	if (Ar.IsSaving() && Ar.IsPersistent())
+	{
+		if (bUseDataRegistry)
+		{
+			// clean out table reference if we are using a data registry
+			Table = nullptr;
+		}
+		else
+		{
+			// clean out the registry type if we aren't using it
+			RegistryType = NAME_None;
+		}
+	}
+	
+	// return false so the normalize serializer will handle the serialization
+	return false;
+}
+
+bool FSoftDataRegistryOrTable::SerializeFromMismatchedTag(const FPropertyTag& Tag, FStructuredArchive::FSlot Slot)
+{
+	// NOTE: this code assumes that the previous soft object ptr was for a UDataTable
+	if (Tag.Type == NAME_SoftObjectProperty)
+	{
+		FSoftObjectPtr OldProperty;
+		Slot << OldProperty;
+
+		bUseDataRegistry = false;
+		Table = OldProperty.ToSoftObjectPath();
+
+		return true;
+	}
+
+	return false;
+}
+
+bool FSoftDataRegistryOrTable::Matches(const UDataTable* InTable) const
+{
+	return (!bUseDataRegistry && (InTable == Table));
+}
+
+bool FSoftDataRegistryOrTable::Matches(const UDataRegistry* InRegistry) const
+{
+	return (bUseDataRegistry && (InRegistry->GetRegistryType() == RegistryType.GetName()));
+}
+
+bool FSoftDataRegistryOrTable::Matches(const FDataRegistryOrTableRow& RegistryOrTableId) const
+{
+	if (RegistryOrTableId.bUseDataRegistryId)
+	{
+		return Matches(RegistryOrTableId.GetDataRegistry());
+	}
+
+	return Matches(RegistryOrTableId.DataTableRow.DataTable);
+}
+
+bool FSoftDataRegistryOrTable::IsValid() const
+{
+	if (bUseDataRegistry)
+	{
+		return (RegistryType.IsValid());
+	}
+
+	return (Table != nullptr);
+}
+
+const UDataRegistry* FSoftDataRegistryOrTable::GetDataRegistry() const
+{
+	if (!bUseDataRegistry)
+	{
+		return nullptr;
+	}
+
+	UDataRegistrySubsystem* RegistrySystem = UDataRegistrySubsystem::Get();
+	if (!ensure(RegistrySystem))
+	{
+		return nullptr;
+	}
+
+	return RegistrySystem->GetRegistryForType(RegistryType);
+}
+
+FDataRegistryOrTableRow::FDataRegistryOrTableRow()
+	: bUseDataRegistryId(false)
+{
+}
+
+FDataRegistryOrTableRow::FDataRegistryOrTableRow(const FDataTableRowHandle& RowHandle)
+	: bUseDataRegistryId(false)
+	, DataTableRow(RowHandle)
+{
+}
+
+FDataRegistryOrTableRow::FDataRegistryOrTableRow(const FDataRegistryId& RegistryId)
+	: bUseDataRegistryId(true)
+	, DataRegistryId(RegistryId)
+{
+}
+
+FString FDataRegistryOrTableRow::ToString() const
+{
+	if (bUseDataRegistryId)
+	{
+		return DataRegistryId.ToString();
+	}
+
+	return DataTableRow.ToDebugString();
+}
+
+const UDataRegistry* FDataRegistryOrTableRow::GetDataRegistry() const
+{
+	if (!bUseDataRegistryId)
+	{
+		return nullptr;
+	}
+
+	UDataRegistrySubsystem* RegistrySystem = UDataRegistrySubsystem::Get();
+	if (!ensure(RegistrySystem))
+	{
+		return nullptr;
+	}
+
+	return RegistrySystem->GetRegistryForType(DataRegistryId.RegistryType);
+}
diff --git a/Engine/Plugins/Runtime/DataRegistry/Source/DataRegistry/Public/DataRegistry.h b/Engine/Plugins/Runtime/DataRegistry/Source/DataRegistry/Public/DataRegistry.h
--- a/Engine/Plugins/Runtime/DataRegistry/Source/DataRegistry/Public/DataRegistry.h
+++ b/Engine/Plugins/Runtime/DataRegistry/Source/DataRegistry/Public/DataRegistry.h
@@ -143,6 +143,66 @@ public:
 		return nullptr;
 	}
 
+	/* Method to run the given predicate on all items */
+	template <class T>
+	void ForEachCachedItem(const FString& ContextString, TFunctionRef<void(const FName& Name, const T& Item)> Predicate) const
+	{
+		TMap<FDataRegistryId, const uint8*> CachedItemMap;
+		const UScriptStruct* ItemTypeStruct = nullptr;
+
+		const FDataRegistryCacheGetResult Result = GetAllCachedItems(CachedItemMap, ItemTypeStruct);
+		if (!Result.WasFound())
+		{
+			UE_LOG(LogDataRegistry, Warning, TEXT("[%hs] No Registry Data found (%s)  Registry:%s"), __FUNCTION__, *ContextString, *RegistryType.ToString());
+			return;
+		}
+
+		if (ItemTypeStruct == nullptr || !ItemTypeStruct->IsChildOf(T::StaticStruct()))
+		{
+			UE_LOG(LogDataRegistry, Warning, TEXT("[%hs] Registry has incorrect row type (%s)  Registry:%s"), __FUNCTION__, *ContextString, *RegistryType.ToString());
+			return;
+		}
+
+		for (TPair<FDataRegistryId, const uint8*>& CachedItem : CachedItemMap)
+		{
+			const T* Item = reinterpret_cast<const T*>(CachedItem.Value);
+			if (Item)
+			{
+				Predicate(*ContextString, *Item);
+			}
+		}
+	}
+
+	/* Method to get all items in the registry */
+	template <class T>
+	void GetAllItems(const TCHAR* ContextString, TArray<const T*>& Items) const
+	{
+		TMap<FDataRegistryId, const uint8*> CachedItemMap;
+		const UScriptStruct* ItemTypeStruct = nullptr;
+
+		const FDataRegistryCacheGetResult Result = GetAllCachedItems(CachedItemMap, ItemTypeStruct);
+		if (!Result.WasFound())
+		{
+			UE_LOG(LogDataRegistry, Warning, TEXT("[%hs] No Registry Data found (%s)  Registry:%s"), __FUNCTION__, ContextString, *RegistryType.ToString());
+			return;
+		}
+
+		if (ItemTypeStruct == nullptr || !ItemTypeStruct->IsChildOf(T::StaticStruct()))
+		{
+			UE_LOG(LogDataRegistry, Warning, TEXT("[%hs] Registry has incorrect row type (%s)  Registry:%s"), __FUNCTION__, ContextString, *RegistryType.ToString());
+			return;
+		}
+
+		for (TPair<FDataRegistryId, const uint8*>& CachedItem : CachedItemMap)
+		{
+			const T* Item = reinterpret_cast<const T*>(CachedItem.Value);
+			Items.Add(Item);
+		}
+	}
+
+	/* Method to get all item names in teh registry */
+	void GetItemNames(TArray<FName>& ItemNames) const;
+
 	/** 
 	 * Fills in a map with all cached (and precached) ids and items for fast iteration. 
 	 * This will use the current resolve context so may not always be valid, and multiple ids can map to the same raw pointer
diff --git a/Engine/Plugins/Runtime/DataRegistry/Source/DataRegistry/Public/DataRegistryTypes.h b/Engine/Plugins/Runtime/DataRegistry/Source/DataRegistry/Public/DataRegistryTypes.h
--- a/Engine/Plugins/Runtime/DataRegistry/Source/DataRegistry/Public/DataRegistryTypes.h
+++ b/Engine/Plugins/Runtime/DataRegistry/Source/DataRegistry/Public/DataRegistryTypes.h
@@ -540,7 +540,7 @@ DECLARE_MULTICAST_DELEGATE_OneParam(FDataRegistryCacheVersionCallback, class UDa
 DECLARE_MULTICAST_DELEGATE(FDataRegistrySubsystemInitializedCallback);
 
 
-DECLARE_LOG_CATEGORY_EXTERN(LogDataRegistry, Log, All);
+DATAREGISTRY_API DECLARE_LOG_CATEGORY_EXTERN(LogDataRegistry, Log, All);
 
 #if UE_ENABLE_INCLUDE_ORDER_DEPRECATED_IN_5_2
 #include "CoreMinimal.h"
diff --git a/Engine/Plugins/Runtime/DataRegistry/Source/DataRegistry/Public/SoftDataRegistryOrTable.h b/Engine/Plugins/Runtime/DataRegistry/Source/DataRegistry/Public/SoftDataRegistryOrTable.h
new file mode 100644
--- /dev/null
+++ b/Engine/Plugins/Runtime/DataRegistry/Source/DataRegistry/Public/SoftDataRegistryOrTable.h
@@ -0,0 +1,210 @@
+// Copyright Epic Games, Inc. All Rights Reserved.
+
+#pragma once
+
+#include "DataRegistry.h"
+#include "DataRegistryTypes.h"
+#include "Engine/DataTable.h"
+
+#include "SoftDataRegistryOrTable.generated.h"
+
+/** 
+ * Defines a DataRegistry or a DataTable with a common interface to both
+ */
+USTRUCT()
+struct DATAREGISTRY_API FSoftDataRegistryOrTable
+{
+	GENERATED_BODY()
+
+	bool Serialize(FArchive& Ar);
+
+	/* Method to check if the given table matches the value in this struct */
+	bool Matches(const UDataTable* InTable) const;
+
+	/* Method to check if the given registry matches the value in this struct */
+	bool Matches(const UDataRegistry* InRegistry) const;
+
+	/* Method to check if the given FDataRegistryOrTableRow matches the value in this struct */
+	bool Matches(const FDataRegistryOrTableRow& RegistryOrTableId) const;
+
+	/* Method to check validity of this struct (Checks that the registry or data table is set)*/
+	bool IsValid() const;
+
+	template <class T>
+	void ForEachItem(const FString& ContextString, TFunctionRef<void(const FName& Name, const T& Item)> Predicate) const
+	{
+		if (bUseDataRegistry)
+		{
+			const UDataRegistry* DataRegistry = GetDataRegistry();
+			if (!DataRegistry)
+			{
+				UE_LOG(LogDataRegistry, Warning, TEXT("[%hs] No Registry found (%s)  Registry:%s"), __FUNCTION__, *ContextString, *RegistryType.GetName().ToString());
+				return;
+			}
+
+			DataRegistry->ForEachCachedItem<T>(ContextString, Predicate);
+		}
+		else
+		{
+			if (!Table)
+			{
+				UE_LOG(LogDataRegistry, Warning, TEXT("[%hs] No Data Table found (%s)"), __FUNCTION__, *ContextString);
+				return;
+			}
+
+			Table->ForeachRow<T>(ContextString, Predicate);
+		}
+	}
+
+	/* This function returns an array of all items in the given registry or data table */
+	template <class T>
+	void GetItems(const TCHAR* ContextString, TArray<const T*>& Items) const
+	{
+		if (bUseDataRegistry)
+		{
+			const UDataRegistry* DataRegistry = GetDataRegistry();
+			if (!DataRegistry)
+			{
+				UE_LOG(LogDataRegistry, Warning, TEXT("[%hs] No Registry found (%s)  Registry:%s"), __FUNCTION__, ContextString, *RegistryType.GetName().ToString());
+				return;
+			}
+
+			DataRegistry->GetAllItems(ContextString, Items);
+		}
+		else
+		{
+			if (!Table)
+			{
+				UE_LOG(LogDataRegistry, Warning, TEXT("[%hs] No Data Table found (%s)"), __FUNCTION__, ContextString);
+				return;
+			}
+
+			Table->GetAllRows<const T>(ContextString, Items);
+		}
+	}
+
+	/* This function returns an array of all items in the given registry or data table */
+	template <class T>
+	void GetAllRows(const FString& ContextString, TArray<T*>& Items) const
+	{
+		GetItems<T>(*ContextString, Items);
+	}
+
+	/* This function returns an array of all item names in the given registry or data table */
+	void GetItemNames(TArray<FName>& ItemNames) const
+	{
+		if (bUseDataRegistry)
+		{
+			const UDataRegistry* DataRegistry = GetDataRegistry();
+			if (!DataRegistry)
+			{
+				UE_LOG(LogDataRegistry, Warning, TEXT("[%hs] No Registry found  Registry:%s"), __FUNCTION__, *RegistryType.GetName().ToString());
+				return;
+			}
+
+			DataRegistry->GetItemNames(ItemNames);
+		}
+		else
+		{
+			if (!Table)
+			{
+				UE_LOG(LogDataRegistry, Warning, TEXT("[%hs] No Data Table found"), __FUNCTION__);
+				return;
+			}
+
+			ItemNames = Table->GetRowNames();
+		}
+	}
+
+	FString GetName() const
+	{
+		if (bUseDataRegistry)
+		{
+			const UDataRegistry* DataRegistry = GetDataRegistry();
+			if (DataRegistry)
+			{
+				return DataRegistry->GetName();
+			}
+		}
+
+		if (Table)
+		{
+			return Table->GetName();
+		}
+
+		return TEXT("");
+	}
+
+	/** Used to upgrade a SoftObjectPtr to a FSoftDataRegistryOrTable */
+	bool SerializeFromMismatchedTag(const FPropertyTag& Tag, FStructuredArchive::FSlot Slot);
+
+	UPROPERTY(EditAnywhere, Category = DataRegistryOrTable)
+	bool bUseDataRegistry = false;
+
+	/** Data Table */
+	UPROPERTY(EditAnywhere, Category = DataRegistryOrTable, meta = (EditCondition = "!bUseDataRegistry", EditConditionHides))
+	TSoftObjectPtr<UDataTable> Table;
+
+	/** Data Registry */
+	UPROPERTY(EditAnywhere, Category = DataRegistryOrTable, meta = (EditCondition = "bUseDataRegistry", EditConditionHides))
+	FDataRegistryType RegistryType;
+
+private:
+
+	/* Method to get the Data Registry */
+	const UDataRegistry* GetDataRegistry() const;
+
+};
+
+template<>
+struct TStructOpsTypeTraits<FSoftDataRegistryOrTable> : public TStructOpsTypeTraitsBase2<FSoftDataRegistryOrTable>
+{
+	enum
+	{
+		WithSerializer = true,
+		WithStructuredSerializeFromMismatchedTag = true,
+	};
+};
+
+/** Defines a DataRegistryId or DataTableRowHandle with a common interface to both */
+USTRUCT(BlueprintType)
+struct DATAREGISTRY_API FDataRegistryOrTableRow
+{
+	GENERATED_BODY()
+
+	FDataRegistryOrTableRow();
+	FDataRegistryOrTableRow(const FDataTableRowHandle& RowHandle);
+	FDataRegistryOrTableRow(const FDataRegistryId& RegistryId);
+
+	/* Method to get the Data Registry for the given registry id */
+	const UDataRegistry* GetDataRegistry() const;
+
+	template <class T>
+	const T* GetItem(const TCHAR* ContextString) const
+	{
+		if (bUseDataRegistryId)
+		{
+			const UDataRegistry* EventRegistry = GetDataRegistry();
+			if (!EventRegistry)
+			{
+				UE_LOG(LogDataRegistry, Warning, TEXT("[%hs] No Registry found  Registry:%s"), __FUNCTION__, *DataRegistryId.RegistryType.GetName().ToString());
+				return nullptr;
+			}
+
+			return EventRegistry->GetCachedItem<T>(DataRegistryId);
+		}
+
+		return DataTableRow.GetRow<T>(ContextString);
+	}
+
+	FString ToString() const;
+
+	UPROPERTY(EditAnywhere, Category = DataRegistryOrTable)
+	bool bUseDataRegistryId = false;
+
+	UPROPERTY(EditAnywhere, BlueprintReadWrite, Category = DataRegistryOrTable, meta = (EditCondition = "!bUseDataRegistryId", EditConditionHides))
+	FDataTableRowHandle DataTableRow;
+	
+	UPROPERTY(EditAnywhere, BlueprintReadWrite, Category = DataRegistryOrTable, meta = (EditCondition = "bUseDataRegistryId", EditConditionHides))
+	FDataRegistryId DataRegistryId;
+};
