SoftDataRegistryOrTable updates to allow async loading and to add other accessors.

diff --git a/Engine/Plugins/Runtime/DataRegistry/Source/DataRegistry/Private/SoftDataRegistryOrTable.cpp b/Engine/Plugins/Runtime/DataRegistry/Source/DataRegistry/Private/SoftDataRegistryOrTable.cpp
--- a/Engine/Plugins/Runtime/DataRegistry/Source/DataRegistry/Private/SoftDataRegistryOrTable.cpp
+++ b/Engine/Plugins/Runtime/DataRegistry/Source/DataRegistry/Private/SoftDataRegistryOrTable.cpp
@@ -3,6 +3,7 @@
 #include "SoftDataRegistryOrTable.h"
 
 #include "DataRegistrySubsystem.h"
+#include "Engine/AssetManager.h"
 
 #include UE_INLINE_GENERATED_CPP_BY_NAME(SoftDataRegistryOrTable)
 
@@ -43,6 +44,38 @@ bool FSoftDataRegistryOrTable::SerializeFromMismatchedTag(const FPropertyTag& Ta
 	return false;
 }
 
+bool FSoftDataRegistryOrTable::IsLoaded() const 
+{ 
+	return bUseDataRegistry || (!Table.IsNull() && !Table.IsPending()); 
+}
+
+void FSoftDataRegistryOrTable::LoadAsync(FStreamableDelegate DelegateToCall)
+{
+	if (bUseDataRegistry || !Table.IsPending())
+	{
+		DelegateToCall.Execute();
+		return;
+	}
+
+	UAssetManager::Get().GetStreamableManager().RequestAsyncLoad(Table.ToSoftObjectPath(), DelegateToCall);
+}
+
+FDataRegistryOrTableRow FSoftDataRegistryOrTable::GetRegistryOrTableRow(FName RowName) const
+{
+	if (bUseDataRegistry)
+	{
+		// make a data registry id
+		return FDataRegistryOrTableRow(FDataRegistryId(RegistryType, RowName));
+	}
+
+	// make a data table row handle
+	FDataTableRowHandle RowHandle;
+	RowHandle.DataTable = Table.Get();
+	RowHandle.RowName = RowName;
+
+	return RowHandle;
+}
+
 bool FSoftDataRegistryOrTable::Matches(const UDataTable* InTable) const
 {
 	return (!bUseDataRegistry && (InTable == Table));
@@ -131,3 +164,8 @@ const UDataRegistry* FDataRegistryOrTableRow::GetDataRegistry() const
 
 	return RegistrySystem->GetRegistryForType(DataRegistryId.RegistryType);
 }
+
+bool FDataRegistryOrTableRow::IsValid() const
+{
+	return bUseDataRegistryId ? DataRegistryId.IsValid() : !DataTableRow.IsNull();
+}
diff --git a/Engine/Plugins/Runtime/DataRegistry/Source/DataRegistry/Public/SoftDataRegistryOrTable.h b/Engine/Plugins/Runtime/DataRegistry/Source/DataRegistry/Public/SoftDataRegistryOrTable.h
--- a/Engine/Plugins/Runtime/DataRegistry/Source/DataRegistry/Public/SoftDataRegistryOrTable.h
+++ b/Engine/Plugins/Runtime/DataRegistry/Source/DataRegistry/Public/SoftDataRegistryOrTable.h
@@ -5,6 +5,7 @@
 #include "DataRegistry.h"
 #include "DataRegistryTypes.h"
 #include "Engine/DataTable.h"
+#include "Engine/StreamableManager.h"
 
 #include "SoftDataRegistryOrTable.generated.h"
 
@@ -85,7 +86,7 @@ struct DATAREGISTRY_API FSoftDataRegistryOrTable
 
 	/* This function returns an array of all items in the given registry or data table */
 	template <class T>
-	void GetAllRows(const FString& ContextString, TArray<T*>& Items) const
+	void GetAllItems(const FString& ContextString, TArray<T*>& Items) const
 	{
 		GetItems<T>(*ContextString, Items);
 	}
@@ -138,6 +139,15 @@ struct DATAREGISTRY_API FSoftDataRegistryOrTable
 	/** Used to upgrade a SoftObjectPtr to a FSoftDataRegistryOrTable */
 	bool SerializeFromMismatchedTag(const FPropertyTag& Tag, FStructuredArchive::FSlot Slot);
 
+	/** Returns whether or not the registry or table is loaded */
+	bool IsLoaded() const;
+
+	/** Async load asset (used for data tables) */
+	void LoadAsync(FStreamableDelegate DelegateToCall);
+
+	/** Method to get a FDataRegistryOrTableRow from the given row name */
+	FDataRegistryOrTableRow GetRegistryOrTableRow(FName RowName) const;
+
 	UPROPERTY(EditAnywhere, Category = DataRegistryOrTable)
 	bool bUseDataRegistry = false;
 
@@ -199,6 +209,9 @@ struct DATAREGISTRY_API FDataRegistryOrTableRow
 
 	FString ToString() const;
 
+	// method to check validity of this row
+	bool IsValid() const;
+
 	UPROPERTY(EditAnywhere, Category = DataRegistryOrTable)
 	bool bUseDataRegistryId = false;
 
