IsValid fix for FSoftDataRegistryOrTable

#rnx

diff --git a/Engine/Plugins/Runtime/DataRegistry/Source/DataRegistry/Private/SoftDataRegistryOrTable.cpp b/Engine/Plugins/Runtime/DataRegistry/Source/DataRegistry/Private/SoftDataRegistryOrTable.cpp
--- a/Engine/Plugins/Runtime/DataRegistry/Source/DataRegistry/Private/SoftDataRegistryOrTable.cpp
+++ b/Engine/Plugins/Runtime/DataRegistry/Source/DataRegistry/Private/SoftDataRegistryOrTable.cpp
@@ -7,6 +7,28 @@
 
 #include UE_INLINE_GENERATED_CPP_BY_NAME(SoftDataRegistryOrTable)
 
+FSoftDataRegistryOrTable::FSoftDataRegistryOrTable()
+	: bUseDataRegistry(false)
+	, Table(nullptr)
+	, RegistryType(NAME_None)
+{
+}
+
+FSoftDataRegistryOrTable::FSoftDataRegistryOrTable(const UDataTable* InDataTable, const FDataRegistryType& InRegistryType)
+{
+	if (InRegistryType != NAME_None)
+	{
+		RegistryType = InRegistryType;
+		Table = nullptr;
+		bUseDataRegistry = true;
+	}
+	else
+	{ 
+		Table = InDataTable;
+		bUseDataRegistry = false;
+	}
+}
+
 bool FSoftDataRegistryOrTable::Serialize(FArchive& Ar)
 {
 	if (Ar.IsSaving() && Ar.IsPersistent())
@@ -103,7 +125,7 @@ bool FSoftDataRegistryOrTable::IsValid() const
 		return (RegistryType.IsValid());
 	}
 
-	return (Table != nullptr);
+	return (!Table.IsNull());
 }
 
 const UDataRegistry* FSoftDataRegistryOrTable::GetDataRegistry() const
diff --git a/Engine/Plugins/Runtime/DataRegistry/Source/DataRegistry/Public/SoftDataRegistryOrTable.h b/Engine/Plugins/Runtime/DataRegistry/Source/DataRegistry/Public/SoftDataRegistryOrTable.h
--- a/Engine/Plugins/Runtime/DataRegistry/Source/DataRegistry/Public/SoftDataRegistryOrTable.h
+++ b/Engine/Plugins/Runtime/DataRegistry/Source/DataRegistry/Public/SoftDataRegistryOrTable.h
@@ -17,6 +17,9 @@ struct DATAREGISTRY_API FSoftDataRegistryOrTable
 {
 	GENERATED_BODY()
 
+	FSoftDataRegistryOrTable();
+	FSoftDataRegistryOrTable(const UDataTable* InDataTable, const FDataRegistryType& InRegistryType);
+
 	bool Serialize(FArchive& Ar);
 
 	/* Method to check if the given table matches the value in this struct */
