Updated LoadAsync in FSoftDataRegistryOrTable (better comment and different execute method)

#rnx

diff --git a/Engine/Plugins/Runtime/DataRegistry/Source/DataRegistry/Private/SoftDataRegistryOrTable.cpp b/Engine/Plugins/Runtime/DataRegistry/Source/DataRegistry/Private/SoftDataRegistryOrTable.cpp
--- a/Engine/Plugins/Runtime/DataRegistry/Source/DataRegistry/Private/SoftDataRegistryOrTable.cpp
+++ b/Engine/Plugins/Runtime/DataRegistry/Source/DataRegistry/Private/SoftDataRegistryOrTable.cpp
@@ -75,7 +75,7 @@ void FSoftDataRegistryOrTable::LoadAsync(FStreamableDelegate DelegateToCall)
 {
 	if (bUseDataRegistry || !Table.IsPending())
 	{
-		DelegateToCall.Execute();
+		FStreamableHandle::ExecuteDelegate(DelegateToCall);
 		return;
 	}
 
diff --git a/Engine/Plugins/Runtime/DataRegistry/Source/DataRegistry/Public/SoftDataRegistryOrTable.h b/Engine/Plugins/Runtime/DataRegistry/Source/DataRegistry/Public/SoftDataRegistryOrTable.h
--- a/Engine/Plugins/Runtime/DataRegistry/Source/DataRegistry/Public/SoftDataRegistryOrTable.h
+++ b/Engine/Plugins/Runtime/DataRegistry/Source/DataRegistry/Public/SoftDataRegistryOrTable.h
@@ -145,8 +145,12 @@ struct DATAREGISTRY_API FSoftDataRegistryOrTable
 	/** Returns whether or not the registry or table is loaded */
 	bool IsLoaded() const;
 
-	/** Async load asset (used for data tables) */
-	void LoadAsync(FStreamableDelegate DelegateToCall);
+	/**
+	 * Requests an async load of a data table using the StreamableManager and then execute the callback, which will happen even if the load fails.
+	 * NOTE: Data registries cannot be loaded this way because they must be registered with the data registry subsystem.
+	 * @param DelegateToCall		Delegate to call when load finishes. Will be called on the next tick if asset is already loaded, or many seconds later
+	 */
+	void LoadAsync(FStreamableDelegate DelegateToCall);	
 
 	/** Method to get a FDataRegistryOrTableRow from the given row name */
 	FDataRegistryOrTableRow GetRegistryOrTableRow(FName RowName) const;
