Fix for FSoftDataRegistryOrTable not using row name in predicate in ForEachItem

#rnx

diff --git a/Engine/Plugins/Runtime/DataRegistry/Source/DataRegistry/Public/DataRegistry.h b/Engine/Plugins/Runtime/DataRegistry/Source/DataRegistry/Public/DataRegistry.h
--- a/Engine/Plugins/Runtime/DataRegistry/Source/DataRegistry/Public/DataRegistry.h
+++ b/Engine/Plugins/Runtime/DataRegistry/Source/DataRegistry/Public/DataRegistry.h
@@ -168,7 +168,7 @@ public:
 			const T* Item = reinterpret_cast<const T*>(CachedItem.Value);
 			if (Item)
 			{
-				Predicate(*ContextString, *Item);
+				Predicate(CachedItem.Key.ItemName, *Item);
 			}
 		}
 	}
