- Moving data registry singleton object out of the Get function and into the class so that we can properly clear it out on Deinitialize otherwise it will cause us to get an invalid pointer to it during shutdown

diff --git a/Engine/Plugins/Runtime/DataRegistry/Source/DataRegistry/Private/DataRegistrySubsystem.cpp b/Engine/Plugins/Runtime/DataRegistry/Source/DataRegistry/Private/DataRegistrySubsystem.cpp
--- a/Engine/Plugins/Runtime/DataRegistry/Source/DataRegistry/Private/DataRegistrySubsystem.cpp
+++ b/Engine/Plugins/Runtime/DataRegistry/Source/DataRegistry/Private/DataRegistrySubsystem.cpp
@@ -20,16 +20,12 @@
 #include "Engine/Engine.h"
 #endif
 
+TObjectPtr<UDataRegistrySubsystem> UDataRegistrySubsystem::SingletonSubSystem = nullptr;
+
 UDataRegistrySubsystem* UDataRegistrySubsystem::Get()
 {
-	// This function is fairly slow and once this is valid it will stay valid until shutdown
-	static UDataRegistrySubsystem* SubSystem = nullptr;
-	if (!SubSystem)
-	{
-		SubSystem = GEngine->GetEngineSubsystem<UDataRegistrySubsystem>();
-	}
-
-	return SubSystem;
+	// GetEngineSubsystem is fairly slow and once SingletonSubSystem is valid it will stay valid until shutdown
+	return SingletonSubSystem;
 }
 
 //static bool GetCachedItemBP(FDataRegistryId ItemId, UPARAM(ref) FTableRowBase& OutItem);
@@ -879,6 +875,9 @@ void UDataRegistrySubsystem::PostAssetManager()
 
 void UDataRegistrySubsystem::Initialize(FSubsystemCollectionBase& Collection)
 {
+	check(SingletonSubSystem == nullptr);
+	SingletonSubSystem = this;
+
 	Super::Initialize(Collection);
 
 	// This should always happen before PostEngineInit
@@ -896,6 +895,8 @@ void UDataRegistrySubsystem::Initialize(FSubsystemCollectionBase& Collection)
 
 void UDataRegistrySubsystem::Deinitialize()
 {
+	check(SingletonSubSystem == this);
+
 	if (!GExitPurge)
 	{
 		DeinitializeAllRegistries();
@@ -905,6 +906,8 @@ void UDataRegistrySubsystem::Deinitialize()
 	bFullyInitialized = false;
 
 	Super::Deinitialize();
+
+	SingletonSubSystem = nullptr;
 }
 
 void UDataRegistrySubsystem::AddReferencedObjects(UObject* InThis, FReferenceCollector& Collector)
diff --git a/Engine/Plugins/Runtime/DataRegistry/Source/DataRegistry/Public/DataRegistrySubsystem.h b/Engine/Plugins/Runtime/DataRegistry/Source/DataRegistry/Public/DataRegistrySubsystem.h
--- a/Engine/Plugins/Runtime/DataRegistry/Source/DataRegistry/Public/DataRegistrySubsystem.h
+++ b/Engine/Plugins/Runtime/DataRegistry/Source/DataRegistry/Public/DataRegistrySubsystem.h
@@ -123,7 +123,7 @@ public:
 
 	// Native interface, works using subsystem instance
 
-	/** Returns the global subsystem instance */
+	/** Returns the global subsystem instance, this can return null during early engine startup and shutdown */
 	static UDataRegistrySubsystem* Get();
 
 	/** Finds the right registry for a type name */
@@ -257,6 +257,9 @@ protected:
 	/** Callback for when the subsystem has finished scanning for and initializing all known data registries */
 	FDataRegistrySubsystemInitializedCallback OnSubsystemInitializedCallback;
 
+	/** Singleton object for the DataRegistrySubsystem::Get function to use, populated in Initialize and cleared out in Deinitialize */
+	static TObjectPtr<UDataRegistrySubsystem> SingletonSubSystem;
+
 #if WITH_EDITOR
 	virtual void PreBeginPIE(bool bStartSimulate);
 	virtual void EndPIE(bool bStartSimulate);
