Added support for mismatched tags for FDataRegistryOrTableRow to convert from FDataTableRowHandle

#rnx

diff --git a/Engine/Plugins/Runtime/DataRegistry/Source/DataRegistry/Private/SoftDataRegistryOrTable.cpp b/Engine/Plugins/Runtime/DataRegistry/Source/DataRegistry/Private/SoftDataRegistryOrTable.cpp
--- a/Engine/Plugins/Runtime/DataRegistry/Source/DataRegistry/Private/SoftDataRegistryOrTable.cpp
+++ b/Engine/Plugins/Runtime/DataRegistry/Source/DataRegistry/Private/SoftDataRegistryOrTable.cpp
@@ -62,6 +62,16 @@ bool FSoftDataRegistryOrTable::SerializeFromMismatchedTag(const FPropertyTag& Ta
 
 		return true;
 	}
+	else if (Tag.Type == NAME_ObjectProperty)
+	{
+		TObjectPtr<UDataTable> OldTable;
+		Slot << OldTable;
+
+		Table = OldTable;
+		bUseDataRegistry = false;
+
+		return true;
+	}
 
 	return false;
 }
@@ -144,6 +154,25 @@ const UDataRegistry* FSoftDataRegistryOrTable::GetDataRegistry() const
 	return RegistrySystem->GetRegistryForType(RegistryType);
 }
 
+bool FDataRegistryOrTableRow::SerializeFromMismatchedTag(const FPropertyTag& Tag, FStructuredArchive::FSlot Slot)
+{
+	static const FLazyName DataTableRowHandleName("DataTableRowHandle");
+	if (Tag.GetType().IsStruct(DataTableRowHandleName))
+	{
+		// Serialize the DataTableRowHandle
+		FDataTableRowHandle OldHandle;
+		FDataTableRowHandle::StaticStruct()->SerializeItem(Slot, &OldHandle, nullptr);
+
+		// copy into new struct
+		bUseDataRegistryId = false;
+		DataTableRow = OldHandle;
+
+		return true;
+	}
+
+	return false;
+}
+
 FDataRegistryOrTableRow::FDataRegistryOrTableRow()
 	: bUseDataRegistryId(false)
 {
diff --git a/Engine/Plugins/Runtime/DataRegistry/Source/DataRegistry/Public/SoftDataRegistryOrTable.h b/Engine/Plugins/Runtime/DataRegistry/Source/DataRegistry/Public/SoftDataRegistryOrTable.h
--- a/Engine/Plugins/Runtime/DataRegistry/Source/DataRegistry/Public/SoftDataRegistryOrTable.h
+++ b/Engine/Plugins/Runtime/DataRegistry/Source/DataRegistry/Public/SoftDataRegistryOrTable.h
@@ -193,6 +193,9 @@ struct DATAREGISTRY_API FDataRegistryOrTableRow
 	FDataRegistryOrTableRow(const FDataTableRowHandle& RowHandle);
 	FDataRegistryOrTableRow(const FDataRegistryId& RegistryId);
 
+	/** Used to upgrade a FDataTableHandle to a FDataRegistryOrTableRow */
+	bool SerializeFromMismatchedTag(const FPropertyTag& Tag, FStructuredArchive::FSlot Slot);
+
 	/* Method to get the Data Registry for the given registry id */
 	const UDataRegistry* GetDataRegistry() const;
 
@@ -228,3 +231,13 @@ struct DATAREGISTRY_API FDataRegistryOrTableRow
 	UPROPERTY(EditAnywhere, BlueprintReadWrite, Category = DataRegistryOrTable, meta = (EditCondition = "bUseDataRegistryId", EditConditionHides))
 	FDataRegistryId DataRegistryId;
 };
+
+template<>
+struct TStructOpsTypeTraits<FDataRegistryOrTableRow> : public TStructOpsTypeTraitsBase2<FDataRegistryOrTableRow>
+{
+	enum
+	{
+		WithStructuredSerializeFromMismatchedTag = true,
+	};
+};
+
