Adding method to retrieve struct type from FDataRegistryOrTableRow

diff --git a/Engine/Plugins/Runtime/DataRegistry/Source/DataRegistry/Private/SoftDataRegistryOrTable.cpp b/Engine/Plugins/Runtime/DataRegistry/Source/DataRegistry/Private/SoftDataRegistryOrTable.cpp
--- a/Engine/Plugins/Runtime/DataRegistry/Source/DataRegistry/Private/SoftDataRegistryOrTable.cpp
+++ b/Engine/Plugins/Runtime/DataRegistry/Source/DataRegistry/Private/SoftDataRegistryOrTable.cpp
@@ -173,6 +173,28 @@ bool FDataRegistryOrTableRow::SerializeFromMismatchedTag(const FPropertyTag& Tag
 	return false;
 }
 
+const UScriptStruct* FDataRegistryOrTableRow::GetStruct() const
+{
+	if (bUseDataRegistryId)
+	{
+		const UDataRegistry* EventRegistry = GetDataRegistry();
+		if (!EventRegistry)
+		{
+			UE_LOG(LogDataRegistry, Warning, TEXT("[%hs] No Registry found  Registry:%s"), __FUNCTION__, *DataRegistryId.RegistryType.GetName().ToString());
+			return nullptr;
+		}
+
+		return EventRegistry->GetItemStruct();
+	}
+
+	if (DataTableRow.DataTable)
+	{
+		return DataTableRow.DataTable->GetRowStruct();
+	}
+
+	return nullptr;
+}
+
 FDataRegistryOrTableRow::FDataRegistryOrTableRow()
 	: bUseDataRegistryId(false)
 {
diff --git a/Engine/Plugins/Runtime/DataRegistry/Source/DataRegistry/Public/SoftDataRegistryOrTable.h b/Engine/Plugins/Runtime/DataRegistry/Source/DataRegistry/Public/SoftDataRegistryOrTable.h
--- a/Engine/Plugins/Runtime/DataRegistry/Source/DataRegistry/Public/SoftDataRegistryOrTable.h
+++ b/Engine/Plugins/Runtime/DataRegistry/Source/DataRegistry/Public/SoftDataRegistryOrTable.h
@@ -217,6 +217,12 @@ struct DATAREGISTRY_API FDataRegistryOrTableRow
 		return DataTableRow.GetRow<T>(ContextString);
 	}
 
+	/**
+	 * Returns the script struct used for the data registry item or data table row. Only works for data registries that
+	 * are registered
+	 */
+	const UScriptStruct* GetStruct() const;
+
 	FString ToString() const;
 
 	// method to check validity of this row
