Adding right-click menu option to open underlying asset (data table, curve table) from data registry editor.
Adding GetSourceAssetPath to UDataRegistrySource.
Making right-click on data registry row highlight the row.

diff --git a/Engine/Plugins/Runtime/DataRegistry/Source/DataRegistry/Private/DataRegistrySource.cpp b/Engine/Plugins/Runtime/DataRegistry/Source/DataRegistry/Private/DataRegistrySource.cpp
--- a/Engine/Plugins/Runtime/DataRegistry/Source/DataRegistry/Private/DataRegistrySource.cpp
+++ b/Engine/Plugins/Runtime/DataRegistry/Source/DataRegistry/Private/DataRegistrySource.cpp
@@ -93,6 +93,11 @@ FString UDataRegistrySource::GetDebugString() const
 	return TEXT("InvalidSource");
 }
 
+FSoftObjectPath UDataRegistrySource::GetSourceAssetPath() const
+{
+	return FSoftObjectPath();
+}
+
 bool UDataRegistrySource::IsTransientSource() const
 {
 	return HasAnyFlags(RF_Transient);
diff --git a/Engine/Plugins/Runtime/DataRegistry/Source/DataRegistry/Private/DataRegistrySource_CurveTable.cpp b/Engine/Plugins/Runtime/DataRegistry/Source/DataRegistry/Private/DataRegistrySource_CurveTable.cpp
--- a/Engine/Plugins/Runtime/DataRegistry/Source/DataRegistry/Private/DataRegistrySource_CurveTable.cpp
+++ b/Engine/Plugins/Runtime/DataRegistry/Source/DataRegistry/Private/DataRegistrySource_CurveTable.cpp
@@ -237,6 +237,11 @@ FString UDataRegistrySource_CurveTable::GetDebugString() const
 	return Super::GetDebugString();
 }
 
+FSoftObjectPath UDataRegistrySource_CurveTable::GetSourceAssetPath() const
+{
+	return SourceTable.ToSoftObjectPath();
+}
+
 bool UDataRegistrySource_CurveTable::Initialize()
 {
 	if (Super::Initialize())
diff --git a/Engine/Plugins/Runtime/DataRegistry/Source/DataRegistry/Private/DataRegistrySource_DataTable.cpp b/Engine/Plugins/Runtime/DataRegistry/Source/DataRegistry/Private/DataRegistrySource_DataTable.cpp
--- a/Engine/Plugins/Runtime/DataRegistry/Source/DataRegistry/Private/DataRegistrySource_DataTable.cpp
+++ b/Engine/Plugins/Runtime/DataRegistry/Source/DataRegistry/Private/DataRegistrySource_DataTable.cpp
@@ -236,6 +236,11 @@ FString UDataRegistrySource_DataTable::GetDebugString() const
 	return Super::GetDebugString();
 }
 
+FSoftObjectPath UDataRegistrySource_DataTable::GetSourceAssetPath() const
+{
+	return SourceTable.ToSoftObjectPath();
+}
+
 bool UDataRegistrySource_DataTable::Initialize()
 {
 	if (Super::Initialize())
diff --git a/Engine/Plugins/Runtime/DataRegistry/Source/DataRegistry/Public/DataRegistrySource.h b/Engine/Plugins/Runtime/DataRegistry/Source/DataRegistry/Public/DataRegistrySource.h
--- a/Engine/Plugins/Runtime/DataRegistry/Source/DataRegistry/Public/DataRegistrySource.h
+++ b/Engine/Plugins/Runtime/DataRegistry/Source/DataRegistry/Public/DataRegistrySource.h
@@ -66,6 +66,9 @@ public:
 	/** Return a useful debug name for this source */
 	virtual FString GetDebugString() const;
 
+	/** Returns the path to the underlying source asset (data table, curve table, etc.) if there is one */
+	virtual FSoftObjectPath GetSourceAssetPath() const;
+
 	/** Returns true if this is a runtime-only source */
 	virtual bool IsTransientSource() const;
 
diff --git a/Engine/Plugins/Runtime/DataRegistry/Source/DataRegistry/Public/DataRegistrySource_CurveTable.h b/Engine/Plugins/Runtime/DataRegistry/Source/DataRegistry/Public/DataRegistrySource_CurveTable.h
--- a/Engine/Plugins/Runtime/DataRegistry/Source/DataRegistry/Public/DataRegistrySource_CurveTable.h
+++ b/Engine/Plugins/Runtime/DataRegistry/Source/DataRegistry/Public/DataRegistrySource_CurveTable.h
@@ -68,6 +68,7 @@ protected:
 	virtual bool AcquireItem(FDataRegistrySourceAcquireRequest&& Request) override;
 	virtual void TimerUpdate(float CurrentTime, float TimerUpdateFrequency) override;
 	virtual FString GetDebugString() const override;
+	virtual FSoftObjectPath GetSourceAssetPath() const override;
 	virtual bool Initialize() override;
 
 	// Object interface
diff --git a/Engine/Plugins/Runtime/DataRegistry/Source/DataRegistry/Public/DataRegistrySource_DataTable.h b/Engine/Plugins/Runtime/DataRegistry/Source/DataRegistry/Public/DataRegistrySource_DataTable.h
--- a/Engine/Plugins/Runtime/DataRegistry/Source/DataRegistry/Public/DataRegistrySource_DataTable.h
+++ b/Engine/Plugins/Runtime/DataRegistry/Source/DataRegistry/Public/DataRegistrySource_DataTable.h
@@ -85,6 +85,7 @@ protected:
 	virtual bool AcquireItem(FDataRegistrySourceAcquireRequest&& Request) override;
 	virtual void TimerUpdate(float CurrentTime, float TimerUpdateFrequency) override;
 	virtual FString GetDebugString() const override;
+	virtual FSoftObjectPath GetSourceAssetPath() const override;
 	virtual bool Initialize() override;
 
 	// Object interface
diff --git a/Engine/Plugins/Runtime/DataRegistry/Source/DataRegistryEditor/Private/SDataRegistryListViewRow.cpp b/Engine/Plugins/Runtime/DataRegistry/Source/DataRegistryEditor/Private/SDataRegistryListViewRow.cpp
--- a/Engine/Plugins/Runtime/DataRegistry/Source/DataRegistryEditor/Private/SDataRegistryListViewRow.cpp
+++ b/Engine/Plugins/Runtime/DataRegistry/Source/DataRegistryEditor/Private/SDataRegistryListViewRow.cpp
@@ -3,6 +3,7 @@
 #include "SDataRegistryListViewRow.h"
 #include "AssetRegistry/AssetIdentifier.h"
 #include "DataRegistryEditorToolkit.h"
+#include "DataRegistrySource.h"
 
 #include "Framework/Application/SlateApplication.h"
 #include "Framework/MultiBox/MultiBoxBuilder.h"
@@ -30,6 +31,7 @@ FReply SDataRegistryListViewRow::OnMouseButtonUp(const FGeometry& MyGeometry, co
 {
 	if (MouseEvent.GetEffectingButton() == EKeys::RightMouseButton && RowDataPtr.IsValid() && FEditorDelegates::OnOpenReferenceViewer.IsBound() && DataRegistryEditor.IsValid())
 	{
+		DataRegistryEditor.Pin().Get()->SetHighlightedRow(RowDataPtr->RowId);
 
 		TSharedRef<SWidget> MenuWidget = MakeRowActionsMenu();
 
@@ -41,27 +43,42 @@ FReply SDataRegistryListViewRow::OnMouseButtonUp(const FGeometry& MyGeometry, co
 	return STableRow::OnMouseButtonUp(MyGeometry, MouseEvent);
 }
 
-void SDataRegistryListViewRow::OnSearchForReferences()
+const FDataRegistrySourceItemId* SDataRegistryListViewRow::FindSource() const
 {
 	if (DataRegistryEditor.IsValid() && RowDataPtr.IsValid())
 	{
-		if (FDataRegistryEditorToolkit* DataRegistryEditorPtr = DataRegistryEditor.Pin().Get())
+		if (const FDataRegistryEditorToolkit* DataRegistryEditorPtr = DataRegistryEditor.Pin().Get())
 		{
-			const FDataRegistrySourceItemId* FoundSource = DataRegistryEditorPtr->GetSourceItemForName(RowDataPtr->RowId);
+			return DataRegistryEditorPtr->GetSourceItemForName(RowDataPtr->RowId);
+		}
+	}
 
-			if (FoundSource)
-			{
-				// TODO Either properly export this in serialize or disable this menu option, similar code in customization
-				TArray<FAssetIdentifier> AssetIdentifiers;
-				AssetIdentifiers.Add(FAssetIdentifier(FoundSource->ItemId.RegistryType, FoundSource->ItemId.ItemName));
+	return nullptr;
+}
 
-				FEditorDelegates::OnOpenReferenceViewer.Broadcast(AssetIdentifiers, FReferenceViewerParams());
-			}
-		}
+void SDataRegistryListViewRow::OnSearchForReferences()
+{
+	if (const FDataRegistrySourceItemId* FoundSource = FindSource())
+	{
+		// TODO Either properly export this in serialize or disable this menu option, similar code in customization
+		TArray<FAssetIdentifier> AssetIdentifiers;
+		AssetIdentifiers.Add(FAssetIdentifier(FoundSource->ItemId.RegistryType, FoundSource->ItemId.ItemName));
+
+		FEditorDelegates::OnOpenReferenceViewer.Broadcast(AssetIdentifiers, FReferenceViewerParams());
 	}
 }
 
 
+void SDataRegistryListViewRow::OnOpenSourceAsset()
+{
+	const FDataRegistrySourceItemId* FoundSource = FindSource();
+	const UDataRegistrySource* SourceObject = FoundSource ? FoundSource->CachedSource.Get() : nullptr;
+	if (SourceObject)
+	{
+		GEditor->GetEditorSubsystem<UAssetEditorSubsystem>()->OpenEditorForAsset(SourceObject->GetSourceAssetPath());
+	}
+}
+
 TSharedRef<SWidget> SDataRegistryListViewRow::GenerateWidgetForColumn(const FName& ColumnName)
 {
 	TSharedPtr<FDataRegistryEditorToolkit> DataRegistryEditorPtr = DataRegistryEditor.Pin();
@@ -166,6 +183,18 @@ TSharedRef<SWidget> SDataRegistryListViewRow::MakeRowActionsMenu()
 		FUIAction(FExecuteAction::CreateSP(this, &SDataRegistryListViewRow::OnSearchForReferences))
 	);
 
+	const FDataRegistrySourceItemId* FoundSource = FindSource();
+	const UDataRegistrySource* SourceObject = FoundSource ? FoundSource->CachedSource.Get() : nullptr;
+	if (SourceObject && !SourceObject->GetSourceAssetPath().IsNull())
+	{
+		MenuBuilder.AddMenuEntry(
+			NSLOCTEXT("FDataRegistryRowUtils", "FDataRegistryRowUtils_OpenSourceAsset", "Open Source Asset"),
+			NSLOCTEXT("FDataRegistryRowUtils", "FDataRegistryRowUtils_OpenSourceAssetTooltip", "Open the asset editor for the source asset of this Row"),
+			FSlateIcon(),
+			FUIAction(FExecuteAction::CreateSP(this, &SDataRegistryListViewRow::OnOpenSourceAsset))
+		);
+	}
+
 	return MenuBuilder.MakeWidget();
 }
 
diff --git a/Engine/Plugins/Runtime/DataRegistry/Source/DataRegistryEditor/Private/SDataRegistryListViewRow.h b/Engine/Plugins/Runtime/DataRegistry/Source/DataRegistryEditor/Private/SDataRegistryListViewRow.h
--- a/Engine/Plugins/Runtime/DataRegistry/Source/DataRegistryEditor/Private/SDataRegistryListViewRow.h
+++ b/Engine/Plugins/Runtime/DataRegistry/Source/DataRegistryEditor/Private/SDataRegistryListViewRow.h
@@ -11,6 +11,7 @@ class STableViewBase;
 
 class FDataRegistryEditorToolkit;
 class SDataRegistryListViewRow;
+struct FDataRegistrySourceItemId;
 
 /**
  * A widget to represent a row in a Data Table Editor widget. This widget allows us to do things like right-click
@@ -42,7 +43,9 @@ public:
 
 private:
 
-	void OnSearchForReferences();	
+	const FDataRegistrySourceItemId* FindSource() const;
+	void OnSearchForReferences();
+	void OnOpenSourceAsset();
 
 	TSharedRef<SWidget> MakeCellWidget(const int32 InRowIndex, const FName& InColumnId);
 
