Reset data and curvetable sources correctly when changing the row structure of a DataRegistry in the editor

diff --git a/Engine/Plugins/Runtime/DataRegistry/Source/DataRegistry/Private/DataRegistrySource_CurveTable.cpp b/Engine/Plugins/Runtime/DataRegistry/Source/DataRegistry/Private/DataRegistrySource_CurveTable.cpp
--- a/Engine/Plugins/Runtime/DataRegistry/Source/DataRegistry/Private/DataRegistrySource_CurveTable.cpp
+++ b/Engine/Plugins/Runtime/DataRegistry/Source/DataRegistry/Private/DataRegistrySource_CurveTable.cpp
@@ -345,6 +345,11 @@ void UDataRegistrySource_CurveTable::PreSave(FObjectPreSaveContext ObjectSaveCon
 	SetCachedTable(true);
 }
 
+void UDataRegistrySource_CurveTable::EditorRefreshSource()
+{
+	SetCachedTable(false);
+}
+
 #endif // WITH_EDITOR
 
 ///////////////////////////////////////////////////////////////////////////////////////////////////////////////////////////////////////////////////////////////
diff --git a/Engine/Plugins/Runtime/DataRegistry/Source/DataRegistry/Private/DataRegistrySource_DataTable.cpp b/Engine/Plugins/Runtime/DataRegistry/Source/DataRegistry/Private/DataRegistrySource_DataTable.cpp
--- a/Engine/Plugins/Runtime/DataRegistry/Source/DataRegistry/Private/DataRegistrySource_DataTable.cpp
+++ b/Engine/Plugins/Runtime/DataRegistry/Source/DataRegistry/Private/DataRegistrySource_DataTable.cpp
@@ -344,6 +344,11 @@ void UDataRegistrySource_DataTable::PreSave(FObjectPreSaveContext ObjectSaveCont
 	SetCachedTable(true);
 }
 
+void UDataRegistrySource_DataTable::EditorRefreshSource()
+{
+	SetCachedTable(false);
+}
+
 #endif // WITH_EDITOR
 
 ///////////////////////////////////////////////////////////////////////////////////////////////////////////////////////////////////////////////////////////////
diff --git a/Engine/Plugins/Runtime/DataRegistry/Source/DataRegistry/Public/DataRegistrySource_CurveTable.h b/Engine/Plugins/Runtime/DataRegistry/Source/DataRegistry/Public/DataRegistrySource_CurveTable.h
--- a/Engine/Plugins/Runtime/DataRegistry/Source/DataRegistry/Public/DataRegistrySource_CurveTable.h
+++ b/Engine/Plugins/Runtime/DataRegistry/Source/DataRegistry/Public/DataRegistrySource_CurveTable.h
@@ -78,6 +78,7 @@ protected:
 
 #if WITH_EDITOR
 	virtual void PreSave(FObjectPreSaveContext ObjectSaveContext) override;
+	virtual void EditorRefreshSource() override;
 #endif
 };
 
diff --git a/Engine/Plugins/Runtime/DataRegistry/Source/DataRegistry/Public/DataRegistrySource_DataTable.h b/Engine/Plugins/Runtime/DataRegistry/Source/DataRegistry/Public/DataRegistrySource_DataTable.h
--- a/Engine/Plugins/Runtime/DataRegistry/Source/DataRegistry/Public/DataRegistrySource_DataTable.h
+++ b/Engine/Plugins/Runtime/DataRegistry/Source/DataRegistry/Public/DataRegistrySource_DataTable.h
@@ -95,6 +95,7 @@ protected:
 
 #if WITH_EDITOR
 	virtual void PreSave(FObjectPreSaveContext ObjectSaveContext) override;
+	virtual void EditorRefreshSource() override;
 #endif
 };
 
