Adding button to open selected data registry asset for FDataRegistryType detail customization

diff --git a/Engine/Plugins/Runtime/DataRegistry/Source/DataRegistryEditor/Private/DataRegistryTypeCustomization.cpp b/Engine/Plugins/Runtime/DataRegistry/Source/DataRegistryEditor/Private/DataRegistryTypeCustomization.cpp
--- a/Engine/Plugins/Runtime/DataRegistry/Source/DataRegistryEditor/Private/DataRegistryTypeCustomization.cpp
+++ b/Engine/Plugins/Runtime/DataRegistry/Source/DataRegistryEditor/Private/DataRegistryTypeCustomization.cpp
@@ -1,9 +1,13 @@
 // Copyright Epic Games, Inc. All Rights Reserved.
 
 #include "DataRegistryTypeCustomization.h"
+#include "DataRegistry.h"
+#include "DataRegistrySubsystem.h"
 #include "DataRegistryEditorModule.h"
 #include "EdGraph/EdGraphSchema.h"
 #include "PropertyCustomizationHelpers.h"
+#include "Widgets/Images/SImage.h"
+#include "Widgets/Input/SButton.h"
 
 #define LOCTEXT_NAMESPACE "DataRegistryEditor"
 
@@ -31,11 +35,97 @@ void FDataRegistryTypeCustomization::CustomizeHeader(TSharedRef<class IPropertyH
 		.MinDesiredWidth(250.0f)
 		.MaxDesiredWidth(0.0f)
 		[
-			PropertyCustomizationHelpers::MakePropertyComboBox(InStructPropertyHandle, FOnGetPropertyComboBoxStrings::CreateStatic(&FDataRegistryEditorModule::GenerateDataRegistryTypeComboBoxStrings, bAllowClear, FilterStructName))
+			SNew(SHorizontalBox)
+
+			+SHorizontalBox::Slot()
+			.AutoWidth()
+			.HAlign(HAlign_Center)
+			.VAlign(VAlign_Center)
+			.Padding(1, 0)
+			[
+				PropertyCustomizationHelpers::MakePropertyComboBox(InStructPropertyHandle, FOnGetPropertyComboBoxStrings::CreateStatic(&FDataRegistryEditorModule::GenerateDataRegistryTypeComboBoxStrings, bAllowClear, FilterStructName))
+			]
+
+			+SHorizontalBox::Slot()
+			.AutoWidth()
+			.HAlign(HAlign_Center)
+			.VAlign(VAlign_Center)
+			.Padding(1, 0)
+			[
+				SNew(SBox)
+				.HAlign(HAlign_Center)
+				.VAlign(VAlign_Center)
+				.WidthOverride(22.0f)
+				.HeightOverride(22.0f)
+				.ToolTipText(this, &FDataRegistryTypeCustomization::GetOpenAssetTooltip)
+				[
+					SNew(SButton)
+					.ButtonStyle(FAppStyle::Get(), "SimpleButton")
+					.OnClicked(this, &FDataRegistryTypeCustomization::OnClickOpenAsset)
+					.ContentPadding(0.0f)
+					.IsFocusable(false)
+					.Visibility(this, &FDataRegistryTypeCustomization::GetOpenAssetVisibility)
+					[
+						SNew(SImage)
+						.Image(FAppStyle::GetBrush("SystemWideCommands.SummonOpenAssetDialog"))
+						.ColorAndOpacity(FSlateColor::UseForeground())
+					]
+				]
+			]
 		];
 	}
 }
 
+const FDataRegistryType* FDataRegistryTypeCustomization::GetPropertyValue() const
+{
+	void* PropertyData;
+	if (StructPropertyHandle && StructPropertyHandle->GetValueData(PropertyData) == FPropertyAccess::Success)
+	{
+		return reinterpret_cast<FDataRegistryType*>(PropertyData);
+	}
+
+	return nullptr;
+}
+
+UDataRegistry* FDataRegistryTypeCustomization::GetDataRegistry() const
+{
+	if (const UDataRegistrySubsystem* DataRegistrySubsystem = UDataRegistrySubsystem::Get())
+	{
+		if (const FDataRegistryType* RegistryType = GetPropertyValue())
+		{
+			return DataRegistrySubsystem->GetRegistryForType(*RegistryType);
+		}
+	}
+
+	return nullptr;
+}
+
+FReply FDataRegistryTypeCustomization::OnClickOpenAsset()
+{
+	if (UDataRegistry* DataRegistry = GetDataRegistry())
+	{
+		GEditor->GetEditorSubsystem<UAssetEditorSubsystem>()->OpenEditorForAsset(DataRegistry);
+	}
+	return FReply::Handled();
+}
+
+EVisibility FDataRegistryTypeCustomization::GetOpenAssetVisibility() const
+{
+	return (GetDataRegistry() != nullptr) ? EVisibility::Visible : EVisibility::Collapsed;
+}
+
+FText FDataRegistryTypeCustomization::GetOpenAssetTooltip() const
+{
+	if (const UDataRegistry* DataRegistry = GetDataRegistry())
+	{
+		FFormatNamedArguments Args;
+		Args.Add(TEXT("Asset"), FText::AsCultureInvariant(DataRegistry->GetName()));
+		return FText::Format(LOCTEXT("OpenSpecificDataRegistry", "Open '{Asset}' in the editor"), Args);
+	}
+
+	return LOCTEXT("OpenDataRegistry", "Open the Data Registry in the editor");
+}
+
 void SDataRegistryTypeGraphPin::Construct(const FArguments& InArgs, UEdGraphPin* InGraphPinObj)
 {
 	SGraphPin::Construct(SGraphPin::FArguments(), InGraphPinObj);
diff --git a/Engine/Plugins/Runtime/DataRegistry/Source/DataRegistryEditor/Private/DataRegistryTypeCustomization.h b/Engine/Plugins/Runtime/DataRegistry/Source/DataRegistryEditor/Private/DataRegistryTypeCustomization.h
--- a/Engine/Plugins/Runtime/DataRegistry/Source/DataRegistryEditor/Private/DataRegistryTypeCustomization.h
+++ b/Engine/Plugins/Runtime/DataRegistry/Source/DataRegistryEditor/Private/DataRegistryTypeCustomization.h
@@ -6,6 +6,9 @@
 #include "IPropertyTypeCustomization.h"
 #include "SGraphPin.h"
 
+class UDataRegistry;
+struct FDataRegistryType;
+
 /** Data Registry type, reads list from subsystem */
 class FDataRegistryTypeCustomization : public IPropertyTypeCustomization
 {
@@ -24,9 +27,20 @@ private:
 	/** Handle to the struct property being customized */
 	TSharedPtr<IPropertyHandle> StructPropertyHandle;
 
+	/** Returns the data of the underlying property */
+	const FDataRegistryType* GetPropertyValue() const;
+
+	/** Returns the data registry associated with this registry type, if there is one */
+	UDataRegistry* GetDataRegistry() const;
+
 	/** Creates the string list */
 	void GenerateComboBoxStrings(TArray< TSharedPtr<FString> >& OutComboBoxStrings, TArray<TSharedPtr<class SToolTip>>& OutToolTips, TArray<bool>& OutRestrictedItems);
 
+	/** Functionality for the "open asset" button */
+	FReply OnClickOpenAsset();
+	EVisibility GetOpenAssetVisibility() const;
+	FText GetOpenAssetTooltip() const;
+
 };
 
 /** Graph pin version of UI */
