[DataRegistry] Depend on UnrealEd if TargetRules.bBuildEditor is true - this matches the same condition used to define WITH_EDITOR, which is used by the module. This addresses an issue where the DataRegistry module would fail to compile if being built by a program setting bCompileAgainstEditor to true.

diff --git a/Engine/Plugins/Runtime/DataRegistry/Source/DataRegistry/DataRegistry.Build.cs b/Engine/Plugins/Runtime/DataRegistry/Source/DataRegistry/DataRegistry.Build.cs
--- a/Engine/Plugins/Runtime/DataRegistry/Source/DataRegistry/DataRegistry.Build.cs
+++ b/Engine/Plugins/Runtime/DataRegistry/Source/DataRegistry/DataRegistry.Build.cs
@@ -18,7 +18,7 @@ namespace UnrealBuildTool.Rules
 			);
 
 			// Needed for PIE callbacks, which should really be somewhere better
-			if (Target.Type == TargetType.Editor)
+			if (Target.bBuildEditor)
 			{
 				PrivateDependencyModuleNames.AddRange(
 					new string[]
