Add GetPossibleDataRegistryIdList to the DataRegistrySubsystem to allow iterating registries from blueprint.
Deprecate GetCachedItemFromLookupBP in favor of new "Find Data Registry Item From Lookup" that uses multiple output pins.
Clean up the comments and names of the BP functions to indicate they are no longer experimental (but may not work on all registries).

diff --git a/Engine/Plugins/Runtime/DataRegistry/Source/DataRegistry/Private/DataRegistrySubsystem.cpp b/Engine/Plugins/Runtime/DataRegistry/Source/DataRegistry/Private/DataRegistrySubsystem.cpp
--- a/Engine/Plugins/Runtime/DataRegistry/Source/DataRegistry/Private/DataRegistrySubsystem.cpp
+++ b/Engine/Plugins/Runtime/DataRegistry/Source/DataRegistry/Private/DataRegistrySubsystem.cpp
@@ -117,6 +117,52 @@ DEFINE_FUNCTION(UDataRegistrySubsystem::execFindCachedItemBP)
 	}
 }
 
+
+//static void FindCachedItemFromLookupBP(FDataRegistryId ItemId, const FDataRegistryLookup& ResolvedLookup, EDataRegistrySubsystemGetItemResult& OutResult, FTableRowBase& OutItem) {}
+DEFINE_FUNCTION(UDataRegistrySubsystem::execFindCachedItemFromLookupBP)
+{
+	P_GET_STRUCT(FDataRegistryId, ItemId);
+	P_GET_STRUCT_REF(FDataRegistryLookup, ResolvedLookup);
+	P_GET_ENUM_REF(EDataRegistrySubsystemGetItemResult, OutResult)
+
+	Stack.MostRecentPropertyAddress = nullptr;
+	Stack.MostRecentPropertyContainer = nullptr;
+	Stack.StepCompiledIn<FStructProperty>(nullptr);
+
+	void* OutItemDataPtr = Stack.MostRecentPropertyAddress;
+	FStructProperty* OutItemProp = CastField<FStructProperty>(Stack.MostRecentProperty);
+	P_FINISH;
+
+	UDataRegistrySubsystem* SubSystem = UDataRegistrySubsystem::Get();
+	check(SubSystem);
+
+	const uint8* CacheData = nullptr;
+	const UScriptStruct* CacheStruct = nullptr;
+	FDataRegistryCacheGetResult CacheResult;
+	OutResult = EDataRegistrySubsystemGetItemResult::NotFound;
+
+	if (OutItemProp && OutItemDataPtr && SubSystem->IsConfigEnabled(true))
+	{
+		P_NATIVE_BEGIN;
+		CacheResult = SubSystem->GetCachedItemRawFromLookup(CacheData, CacheStruct, ItemId, ResolvedLookup);
+
+		if (CacheResult && CacheStruct && CacheData)
+		{
+			UScriptStruct* OutputStruct = OutItemProp->Struct;
+
+			const bool bCompatible = (OutputStruct == CacheStruct) ||
+				(OutputStruct->IsChildOf(CacheStruct) && FStructUtils::TheSameLayout(OutputStruct, CacheStruct));
+
+			if (bCompatible)
+			{
+				OutResult = EDataRegistrySubsystemGetItemResult::Found;
+				CacheStruct->CopyScriptStruct(OutItemDataPtr, CacheData);
+			}
+		}
+		P_NATIVE_END;
+	}
+}
+
 //static bool GetCachedItemFromLookupBP(FDataRegistryId ItemId, const FDataRegistryLookup& ResolvedLookup, UPARAM(ref) FTableRowBase& OutItem) { return false; }
 DEFINE_FUNCTION(UDataRegistrySubsystem::execGetCachedItemFromLookupBP)
 {
@@ -770,6 +816,27 @@ bool UDataRegistrySubsystem::AcquireItemBP(FDataRegistryId ItemId, FDataRegistry
 		}));
 }
 
+void UDataRegistrySubsystem::GetPossibleDataRegistryIdList(FDataRegistryType RegistryType, TArray<FDataRegistryId>& OutIdList)
+{
+	OutIdList.Reset();
+
+	UDataRegistrySubsystem* SubSystem = UDataRegistrySubsystem::Get();
+	check(SubSystem);
+
+	if (!SubSystem->IsConfigEnabled(true))
+	{
+		return;
+	}
+
+	UDataRegistry* FoundRegistry = SubSystem->GetRegistryForType(RegistryType);
+
+	if (FoundRegistry)
+	{
+		// Don't sort by default
+		FoundRegistry->GetPossibleRegistryIds(OutIdList, false);
+	}
+}
+
 void UDataRegistrySubsystem::EvaluateDataRegistryCurve(FDataRegistryId ItemId, float InputValue, float DefaultValue, EDataRegistrySubsystemGetItemResult& OutResult, float& OutValue)
 {
 	UDataRegistrySubsystem* SubSystem = UDataRegistrySubsystem::Get();
diff --git a/Engine/Plugins/Runtime/DataRegistry/Source/DataRegistry/Public/DataRegistrySubsystem.h b/Engine/Plugins/Runtime/DataRegistry/Source/DataRegistry/Public/DataRegistrySubsystem.h
--- a/Engine/Plugins/Runtime/DataRegistry/Source/DataRegistry/Public/DataRegistrySubsystem.h
+++ b/Engine/Plugins/Runtime/DataRegistry/Source/DataRegistry/Public/DataRegistrySubsystem.h
@@ -11,7 +11,7 @@ struct FRealCurve;
 struct FTableRowBase;
 
 
-/** Enum used to indicate success or failure of EvaluateCurveTableRow. */
+/** Enum used to indicate success or failure of finding a data registry item */
 UENUM()
 enum class EDataRegistrySubsystemGetItemResult : uint8
 {
@@ -31,51 +31,68 @@ public:
 	// Blueprint Interface, it is static for ease of use in custom nodes
 
 	/**
-	 * Attempts to get cached structure data stored in a DataRegistry, modifying OutItem if the item is available
-	 * (EXPERIMENTAL) this version has an input param and simple bool return
+	 * Attempts to get cached structure data stored in a DataRegistry, modifying OutItem if the item is available.
+	 * This version has an input param and simple bool return.
 	 *
 	 * @param ItemID		Item identifier to lookup in cache
 	 * @param OutItem		This must be the same type as the registry, if the item is found this will be filled in with the found data
 	 * @returns				Returns true if the item was found and OutItem was modified
 	 */
-	UFUNCTION(BlueprintCallable, CustomThunk, Category = DataRegistry, meta = (DisplayName = "Get Data Registry Item (experimental)", CustomStructureParam = "OutItem"))
+	UFUNCTION(BlueprintCallable, CustomThunk, Category = DataRegistry, meta = (DisplayName = "Get Data Registry Item", CustomStructureParam = "OutItem"))
 	static bool GetCachedItemBP(FDataRegistryId ItemId, UPARAM(ref) FTableRowBase& OutItem) { return false; }
 	DECLARE_FUNCTION(execGetCachedItemBP);
 
 	/**
-	 * Attempts to get cached structure data stored in a DataRegistry, modifying OutItem if the item is available
-	 * (EXPERIMENTAL) this version has an output param and enum result
+	 * Attempts to get cached structure data stored in a DataRegistry, returning OutItem if the item is available.
+	 * This version has two output pins for convenience, and OutItem should not be accessed from the Not Found pin.
 	 *
 	 * @param ItemID		Item identifier to lookup in cache
 	 * @param OutItem		This must be the same type as the registry, if the item is found this will be filled in with the found data
-	 * @returns				Returns true if the item was found and OutItem was modified
+	 * @param OutResult		Pick execution pin based on if the item was found and OutItem is valid
 	 */
-	UFUNCTION(BlueprintCallable, CustomThunk, Category = DataRegistry, meta = (DisplayName = "Find Data Registry Item (experimental)", CustomStructureParam = "OutItem", ExpandEnumAsExecs = "OutResult"))
+	UFUNCTION(BlueprintCallable, CustomThunk, Category = DataRegistry, meta = (DisplayName = "Find Data Registry Item", CustomStructureParam = "OutItem", ExpandEnumAsExecs = "OutResult"))
 	static void FindCachedItemBP(FDataRegistryId ItemId, EDataRegistrySubsystemGetItemResult& OutResult, FTableRowBase& OutItem) {}
 	DECLARE_FUNCTION(execFindCachedItemBP);
 
 	/**
-	 * Attempts to get structure data stored in a DataRegistry cache after an async acquire, modifying OutItem if the item is available
+	 * Attempts to get structure data stored in a DataRegistry cache after an async acquire, returning OutItem if the item is available.
+	 * OutItem should not be accessed from the Not Found pin was not found.
 	 *
 	 * @param ItemID			Item identifier to lookup in cache
 	 * @param ResolvedLookup	Resolved identifier returned by acquire function
 	 * @param OutItem			This must be the same type as the registry, if the item is found this will be filled in with the found data
-	 * @returns					Returns true if the item was found and OutItem was modified
+	 * @param OutResult			Pick execution pin based on if the item was found and OutItem is valid
 	 */
-	UFUNCTION(BlueprintCallable, CustomThunk, Category = DataRegistry, meta = (DisplayName = "Get Data Registry Item From Lookup (experimental)", CustomStructureParam = "OutItem"))
+	UFUNCTION(BlueprintCallable, CustomThunk, Category = DataRegistry, meta = (DisplayName = "Find Data Registry Item From Lookup", CustomStructureParam = "OutItem", ExpandEnumAsExecs = "OutResult"))
+	static void FindCachedItemFromLookupBP(FDataRegistryId ItemId, const FDataRegistryLookup& ResolvedLookup, EDataRegistrySubsystemGetItemResult& OutResult, FTableRowBase& OutItem) {}
+	DECLARE_FUNCTION(execFindCachedItemFromLookupBP);
+
+	/** Deprecated in favor of FindCachedItemFromLookupBP, but does still work properly */
+	UFUNCTION(BlueprintCallable, CustomThunk, Category = DataRegistry, meta = (DeprecatedFunction = "Note", DeprecationMessage = "Use Find Data Registry Item From Lookup instead", DisplayName = "Get Data Registry Item From Lookup", CustomStructureParam = "OutItem"))
 	static bool GetCachedItemFromLookupBP(FDataRegistryId ItemId, const FDataRegistryLookup& ResolvedLookup, FTableRowBase& OutItem) { return false; }
 	DECLARE_FUNCTION(execGetCachedItemFromLookupBP);
 
 	/**
-	 * Starts an asynchronous acquire of a data registry item that may not yet be cached.
+	 * Starts an asynchronous acquire of a data registry item that may not yet be cached, and then accessed with Get Data Registry Item From Lookup
+	 * This function will only work properly if the data registry is set up for asynchronous querying.
 	 *
 	 * @param ItemID			Item identifier to lookup in cache
 	 * @param AcquireCallback	Delegate that will be called after acquire succeeds or failed
 	 * @returns					Returns true if request was started, false on unrecoverable error
 	 */
-	UFUNCTION(BlueprintCallable, Category = DataRegistry, meta = (DisplayName = "Acquire Data Registry Item (experimental)") )
+	UFUNCTION(BlueprintCallable, Category = DataRegistry, meta = (DisplayName = "Acquire Data Registry Item") )
 	static bool AcquireItemBP(FDataRegistryId ItemId, FDataRegistryItemAcquiredBPCallback AcquireCallback);
 
+	/** 
+	 * Returns the list of known identifiers for an active data registry so they can be iterated with Find or Acquire.
+	 * Depending on how the registry is setup, this could be a large number of identifiers and they may not all be available.
+	 * 
+	 * @param RegistryType	The type of data registry to query
+	 * @param OutIdList		The list of known identifiers for the type, which will be empty if the type is not registered
+	 */
+	UFUNCTION(BlueprintCallable, BlueprintPure = false, Category = DataRegistry, meta = (DisplayName = "Get Possible Data Registry Id List"))
+	static void GetPossibleDataRegistryIdList(FDataRegistryType RegistryType, TArray<FDataRegistryId>& OutIdList);
+
 	/**
 	 * Attempts to evaluate a curve stored in a DataRegistry cache using a specific input value
 	 *
