UE: Add an option to delay DataRegistrySubsystem loads until the first PIE

diff --git a/Engine/Plugins/Runtime/DataRegistry/Source/DataRegistry/Private/DataRegistrySubsystem.cpp b/Engine/Plugins/Runtime/DataRegistry/Source/DataRegistry/Private/DataRegistrySubsystem.cpp
--- a/Engine/Plugins/Runtime/DataRegistry/Source/DataRegistry/Private/DataRegistrySubsystem.cpp
+++ b/Engine/Plugins/Runtime/DataRegistry/Source/DataRegistry/Private/DataRegistrySubsystem.cpp
@@ -937,7 +937,19 @@ void UDataRegistrySubsystem::PostAssetManager()
 {
 	// This will happen fast in game builds, can be many seconds in editor builds due to async load of asset registry
 	bReadyForInitialization = true;
-	LoadAllRegistries();
+
+	bool bLoadAllRegistries = true;
+#if WITH_EDITOR
+	if (GIsEditor && GetDefault<UDataRegistrySettings>()->bDelayLoadingDataRegistriesUntilPIE)
+	{
+		bLoadAllRegistries = false;
+		bLoadAllRegistriesOnNextPIE = true;
+	}
+#endif
+	if (bLoadAllRegistries)
+	{
+		LoadAllRegistries();
+	}
 }
 
 void UDataRegistrySubsystem::Initialize(FSubsystemCollectionBase& Collection)
@@ -1006,6 +1018,12 @@ void UDataRegistrySubsystem::PreBeginPIE(bool bStartSimulate)
 {
 	if (IsConfigEnabled())
 	{
+		if (bLoadAllRegistriesOnNextPIE)
+		{
+			LoadAllRegistries();
+			bLoadAllRegistriesOnNextPIE = false;
+		}
+
 		RefreshRegistryMap();
 		InitializeAllRegistries(true);
 	}
diff --git a/Engine/Plugins/Runtime/DataRegistry/Source/DataRegistry/Public/DataRegistrySettings.h b/Engine/Plugins/Runtime/DataRegistry/Source/DataRegistry/Public/DataRegistrySettings.h
--- a/Engine/Plugins/Runtime/DataRegistry/Source/DataRegistry/Public/DataRegistrySettings.h
+++ b/Engine/Plugins/Runtime/DataRegistry/Source/DataRegistry/Public/DataRegistrySettings.h
@@ -30,6 +30,10 @@ public:
 	UPROPERTY(config, EditAnywhere, Category = "Data Registry")
 	bool bIgnoreMissingCookedAssetRegistryData = false;
 
+	/** If true, in the editor data registry assets will be loaded before the first PIE instead of during startup */
+	UPROPERTY(config, EditAnywhere, Category = "Data Registry")
+	bool bDelayLoadingDataRegistriesUntilPIE = true;
+
 
 	/** Return true if we are allowed to ignore missing asset registry data based on settings and build */
 	bool CanIgnoreMissingAssetData() const;
diff --git a/Engine/Plugins/Runtime/DataRegistry/Source/DataRegistry/Public/DataRegistrySubsystem.h b/Engine/Plugins/Runtime/DataRegistry/Source/DataRegistry/Public/DataRegistrySubsystem.h
--- a/Engine/Plugins/Runtime/DataRegistry/Source/DataRegistry/Public/DataRegistrySubsystem.h
+++ b/Engine/Plugins/Runtime/DataRegistry/Source/DataRegistry/Public/DataRegistrySubsystem.h
@@ -278,6 +278,7 @@ protected:
 	static TObjectPtr<UDataRegistrySubsystem> SingletonSubSystem;
 
 #if WITH_EDITOR
+	bool bLoadAllRegistriesOnNextPIE = false;
 	virtual void PreBeginPIE(bool bStartSimulate);
 	virtual void EndPIE(bool bStartSimulate);
 #endif
