UE: Disable bDelayLoadingDataRegistriesUntilPIE. This was meant to be defaulted off

diff --git a/Engine/Plugins/Runtime/DataRegistry/Source/DataRegistry/Public/DataRegistrySettings.h b/Engine/Plugins/Runtime/DataRegistry/Source/DataRegistry/Public/DataRegistrySettings.h
--- a/Engine/Plugins/Runtime/DataRegistry/Source/DataRegistry/Public/DataRegistrySettings.h
+++ b/Engine/Plugins/Runtime/DataRegistry/Source/DataRegistry/Public/DataRegistrySettings.h
@@ -32,7 +32,7 @@ public:
 
 	/** If true, in the editor data registry assets will be loaded before the first PIE instead of during startup */
 	UPROPERTY(config, EditAnywhere, Category = "Data Registry")
-	bool bDelayLoadingDataRegistriesUntilPIE = true;
+	bool bDelayLoadingDataRegistriesUntilPIE = false;
 
 
 	/** Return true if we are allowed to ignore missing asset registry data based on settings and build */
