UE: Do not delay data registry loading in commandlets because they do not PIE

diff --git a/Engine/Plugins/Runtime/DataRegistry/Source/DataRegistry/Private/DataRegistrySubsystem.cpp b/Engine/Plugins/Runtime/DataRegistry/Source/DataRegistry/Private/DataRegistrySubsystem.cpp
--- a/Engine/Plugins/Runtime/DataRegistry/Source/DataRegistry/Private/DataRegistrySubsystem.cpp
+++ b/Engine/Plugins/Runtime/DataRegistry/Source/DataRegistry/Private/DataRegistrySubsystem.cpp
@@ -940,7 +940,7 @@ void UDataRegistrySubsystem::PostAssetManager()
 
 	bool bLoadAllRegistries = true;
 #if WITH_EDITOR
-	if (GIsEditor && GetDefault<UDataRegistrySettings>()->bDelayLoadingDataRegistriesUntilPIE)
+	if (GIsEditor && !IsRunningCommandlet() && GetDefault<UDataRegistrySettings>()->bDelayLoadingDataRegistriesUntilPIE)
 	{
 		bLoadAllRegistries = false;
 		bLoadAllRegistriesOnNextPIE = true;
