UE: Lazy flush data registry deferred load if any registries are requested before PIE

diff --git a/Engine/Plugins/Runtime/DataRegistry/Source/DataRegistry/Private/DataRegistrySubsystem.cpp b/Engine/Plugins/Runtime/DataRegistry/Source/DataRegistry/Private/DataRegistrySubsystem.cpp
--- a/Engine/Plugins/Runtime/DataRegistry/Source/DataRegistry/Private/DataRegistrySubsystem.cpp
+++ b/Engine/Plugins/Runtime/DataRegistry/Source/DataRegistry/Private/DataRegistrySubsystem.cpp
@@ -210,6 +210,8 @@ DEFINE_FUNCTION(UDataRegistrySubsystem::execGetCachedItemFromLookupBP)
 
 UDataRegistry* UDataRegistrySubsystem::GetRegistryForType(FName RegistryType) const
 {
+	FlushDeferredRegistryLoad(false);
+
 	auto* FoundRegistry = RegistryMap.Find(RegistryType);
 
 	if (FoundRegistry && *FoundRegistry)
@@ -241,6 +243,8 @@ void UDataRegistrySubsystem::GetAllRegistries(TArray<UDataRegistry*>& AllRegistr
 {
 	AllRegistries.Reset();
 
+	FlushDeferredRegistryLoad(false);
+
 	for (const FRegistryMapPair& RegistryPair : RegistryMap)
 	{
 		UDataRegistry* Registry = RegistryPair.Value;
@@ -534,6 +538,8 @@ bool UDataRegistrySubsystem::RegisterSpecificAsset(FDataRegistryType RegistryTyp
 
 	if (!RegistryType.IsValid())
 	{
+		FlushDeferredRegistryLoad(false);
+
 		bool bMadeChange = false;
 		for (const FRegistryMapPair& RegistryPair : RegistryMap)
 		{
@@ -938,15 +944,14 @@ void UDataRegistrySubsystem::PostAssetManager()
 	// This will happen fast in game builds, can be many seconds in editor builds due to async load of asset registry
 	bReadyForInitialization = true;
 
-	bool bLoadAllRegistries = true;
 #if WITH_EDITOR
 	if (GIsEditor && !IsRunningCommandlet() && GetDefault<UDataRegistrySettings>()->bDelayLoadingDataRegistriesUntilPIE)
 	{
-		bLoadAllRegistries = false;
+		bReadyForInitialization = false;
 		bLoadAllRegistriesOnNextPIE = true;
 	}
 #endif
-	if (bLoadAllRegistries)
+	if (bReadyForInitialization)
 	{
 		LoadAllRegistries();
 	}
@@ -1016,14 +1021,10 @@ void UDataRegistrySubsystem::AddReferencedObjects(UObject* InThis, FReferenceCol
 
 void UDataRegistrySubsystem::PreBeginPIE(bool bStartSimulate)
 {
+	FlushDeferredRegistryLoad(true);
+
 	if (IsConfigEnabled())
 	{
-		if (bLoadAllRegistriesOnNextPIE)
-		{
-			LoadAllRegistries();
-			bLoadAllRegistriesOnNextPIE = false;
-		}
-
 		RefreshRegistryMap();
 		InitializeAllRegistries(true);
 	}
@@ -1038,6 +1039,24 @@ void UDataRegistrySubsystem::EndPIE(bool bStartSimulate)
 }
 #endif
 
+void UDataRegistrySubsystem::FlushDeferredRegistryLoad(bool bExpected) const
+{
+#if WITH_EDITOR
+	if (bLoadAllRegistriesOnNextPIE)
+	{
+		// Const cast because the flush can be lazy on first access.
+		UDataRegistrySubsystem* MutableThis = const_cast<UDataRegistrySubsystem*>(this);
+		if (!bExpected)
+		{
+			UE_LOG(LogDataRegistry, Warning, TEXT("Accessed data registry map before PIE. Loading on demand."));
+		}
+		
+		MutableThis->bReadyForInitialization = true;
+		MutableThis->LoadAllRegistries();
+		MutableThis->bLoadAllRegistriesOnNextPIE = false;
+	}
+#endif // WITH_EDITOR
+}
 
 
 static FAutoConsoleCommand CVarDumpRegistryTypeSummary(
diff --git a/Engine/Plugins/Runtime/DataRegistry/Source/DataRegistry/Public/DataRegistrySubsystem.h b/Engine/Plugins/Runtime/DataRegistry/Source/DataRegistry/Public/DataRegistrySubsystem.h
--- a/Engine/Plugins/Runtime/DataRegistry/Source/DataRegistry/Public/DataRegistrySubsystem.h
+++ b/Engine/Plugins/Runtime/DataRegistry/Source/DataRegistry/Public/DataRegistrySubsystem.h
@@ -282,6 +282,7 @@ protected:
 	virtual void PreBeginPIE(bool bStartSimulate);
 	virtual void EndPIE(bool bStartSimulate);
 #endif
+	void FlushDeferredRegistryLoad(bool bExpected) const;
 };
 
 /* Test actor, move later
