UE: Null check ItemStruct in DataRegistrySource_DataTable, since the struct could be unassigned or a missing asset

diff --git a/Engine/Plugins/Runtime/DataRegistry/Source/DataRegistry/Private/DataRegistrySource_DataTable.cpp b/Engine/Plugins/Runtime/DataRegistry/Source/DataRegistry/Private/DataRegistrySource_DataTable.cpp
--- a/Engine/Plugins/Runtime/DataRegistry/Source/DataRegistry/Private/DataRegistrySource_DataTable.cpp
+++ b/Engine/Plugins/Runtime/DataRegistry/Source/DataRegistry/Private/DataRegistrySource_DataTable.cpp
@@ -402,33 +402,35 @@ bool UMetaDataRegistrySource_DataTable::DoesAssetPassFilter(const FAssetData& As
 	FString RowStructureString;
 	if (AssetData.GetTagValue(RowStructureTagName, RowStructureString))
 	{
-		const UScriptStruct* ItemStruct = GetItemStruct();
-		if (RowStructureString == ItemStruct->GetName() || RowStructureString == ItemStruct->GetStructPathName().ToString())
-		{
-			return true;
-		}
-		else
+		if (const UScriptStruct* ItemStruct = GetItemStruct())
 		{
-			// TODO no 100% way to check for inherited row structs, but BP types can't inherit anyway
-			const UScriptStruct* RowStruct = FindFirstObject<UScriptStruct>(*RowStructureString, EFindFirstObjectOptions::ExactClass);
-
-			// Check if the row struct is a child of the item struct
-			if (RowStruct != nullptr)
+			if (RowStructureString == ItemStruct->GetName() || RowStructureString == ItemStruct->GetStructPathName().ToString())
 			{
-				if (RowStruct->IsChildOf(ItemStruct))
-				{
-					return true;
-				}	
+				return true;
 			}
-			// Otherwise check if the row struct has been redirected to the item struct
 			else
 			{
-				FName RowStructureName = FName(FPackageName::ObjectPathToObjectName(RowStructureString)); // RobM this should use path names as well
+				// TODO no 100% way to check for inherited row structs, but BP types can't inherit anyway
+				const UScriptStruct* RowStruct = FindFirstObject<UScriptStruct>(*RowStructureString, EFindFirstObjectOptions::ExactClass);
+
+				// Check if the row struct is a child of the item struct
+				if (RowStruct != nullptr)
+				{
+					if (RowStruct->IsChildOf(ItemStruct))
+					{
+						return true;
+					}	
+				}
+				// Otherwise check if the row struct has been redirected to the item struct
+				else
+				{
+					FName RowStructureName = FName(FPackageName::ObjectPathToObjectName(RowStructureString)); // RobM this should use path names as well
 				
-				TArray<FCoreRedirectObjectName> PreviousNames;
-				FCoreRedirects::FindPreviousNames(ECoreRedirectFlags::Type_Struct, ItemStruct->GetPathName(), PreviousNames);
+					TArray<FCoreRedirectObjectName> PreviousNames;
+					FCoreRedirects::FindPreviousNames(ECoreRedirectFlags::Type_Struct, ItemStruct->GetPathName(), PreviousNames);
 				
-				return PreviousNames.ContainsByPredicate([&RowStructureName](const FCoreRedirectObjectName& PreviousName){ return PreviousName.ObjectName == RowStructureName; });
+					return PreviousNames.ContainsByPredicate([&RowStructureName](const FCoreRedirectObjectName& PreviousName){ return PreviousName.ObjectName == RowStructureName; });
+				}
 			}
 		}
 	}
