Changing log verbosity of message in UDataRegistrySubsystem::ApplyPreregisterMap

#rnx
#tests launching pie session and hitting log message

diff --git a/Engine/Plugins/Runtime/DataRegistry/Source/DataRegistry/Private/DataRegistrySubsystem.cpp b/Engine/Plugins/Runtime/DataRegistry/Source/DataRegistry/Private/DataRegistrySubsystem.cpp
--- a/Engine/Plugins/Runtime/DataRegistry/Source/DataRegistry/Private/DataRegistrySubsystem.cpp
+++ b/Engine/Plugins/Runtime/DataRegistry/Source/DataRegistry/Private/DataRegistrySubsystem.cpp
@@ -742,7 +742,7 @@ void UDataRegistrySubsystem::ApplyPreregisterMap(UDataRegistry* Registry)
 			}
 			else if (RegisterAssetResult == EDataRegistryRegisterAssetResult::AssetAlreadyRegistered)
 			{
-				UE_LOG(LogDataRegistry, Log, TEXT("ApplyPreregisterMap did not register %s with %s, this asset was already registered with the data registry."), *AssetPath.ToString(), *Registry->GetRegistryType().ToString());
+				UE_LOG(LogDataRegistry, Verbose, TEXT("ApplyPreregisterMap did not register %s with %s, this asset was already registered with the data registry."), *AssetPath.ToString(), *Registry->GetRegistryType().ToString());
 			}
 		}
 	}
