Update SoftDataRegistryOrTable to be able to serialize from a mismatched tag that is a registry type

#rnx

diff --git a/Engine/Plugins/Runtime/DataRegistry/Source/DataRegistry/Private/SoftDataRegistryOrTable.cpp b/Engine/Plugins/Runtime/DataRegistry/Source/DataRegistry/Private/SoftDataRegistryOrTable.cpp
--- a/Engine/Plugins/Runtime/DataRegistry/Source/DataRegistry/Private/SoftDataRegistryOrTable.cpp
+++ b/Engine/Plugins/Runtime/DataRegistry/Source/DataRegistry/Private/SoftDataRegistryOrTable.cpp
@@ -51,6 +51,19 @@ bool FSoftDataRegistryOrTable::Serialize(FArchive& Ar)
 
 bool FSoftDataRegistryOrTable::SerializeFromMismatchedTag(const FPropertyTag& Tag, FStructuredArchive::FSlot Slot)
 {
+	// Note: This assumes the previous slot was an FDataRegistryType
+	if (Tag.GetType().IsStruct(FDataRegistryType::StaticStruct()->GetFName()))
+	{
+		FDataRegistryType Reference;
+		FDataRegistryType::StaticStruct()->SerializeItem(Slot, &Reference, nullptr);
+		if (!Reference.GetName().IsNone())
+		{
+			RegistryType = Reference;
+			bUseDataRegistry = true;
+			return true;
+		}
+	}
+	
 	// NOTE: this code assumes that the previous soft object ptr was for a UDataTable
 	if (Tag.Type == NAME_SoftObjectProperty)
 	{
