Made FSoftDataRegistryOrTable a BP type so that it can be used as a parameter in StateTrees

diff --git a/Engine/Plugins/Runtime/DataRegistry/Source/DataRegistry/Public/SoftDataRegistryOrTable.h b/Engine/Plugins/Runtime/DataRegistry/Source/DataRegistry/Public/SoftDataRegistryOrTable.h
--- a/Engine/Plugins/Runtime/DataRegistry/Source/DataRegistry/Public/SoftDataRegistryOrTable.h
+++ b/Engine/Plugins/Runtime/DataRegistry/Source/DataRegistry/Public/SoftDataRegistryOrTable.h
@@ -12,7 +12,7 @@
 /** 
  * Defines a DataRegistry or a DataTable with a common interface to both
  */
-USTRUCT()
+USTRUCT(BlueprintType)
 struct FSoftDataRegistryOrTable
 {
 	GENERATED_BODY()
