Workaround for a musttail PGO build bug in clang in the UDataRegistry::TimerUpdate()

diff --git a/Engine/Plugins/Runtime/DataRegistry/Source/DataRegistry/Private/DataRegistry.cpp b/Engine/Plugins/Runtime/DataRegistry/Source/DataRegistry/Private/DataRegistry.cpp
--- a/Engine/Plugins/Runtime/DataRegistry/Source/DataRegistry/Private/DataRegistry.cpp
+++ b/Engine/Plugins/Runtime/DataRegistry/Source/DataRegistry/Private/DataRegistry.cpp
@@ -991,7 +991,18 @@ float UDataRegistry::GetCurrentTime()
 	return (float)(FApp::GetCurrentTime() - GStartTime);
 }
 
+// this is a workaround for a bug in clang that applies musttail attribute to Source->TimerUpdate call in some cases
+#if defined(__has_attribute)
+#	if __has_attribute(no_profile_instrument_function)
+#		define NO_PROFILE_ATTRIBUTE __attribute__((no_profile_instrument_function))
+#	endif
+#endif
+
+#ifndef NO_PROFILE_ATTRIBUTE
+#	define NO_PROFILE_ATTRIBUTE
+#endif
 
+NO_PROFILE_ATTRIBUTE
 void UDataRegistry::TimerUpdate()
 {
 	if (!bIsInitialized)
