Added UE_NO_PROFILE_ATTRIBUTE to avoid functions from being profiled by clang's PGO

diff --git a/Engine/Plugins/Runtime/DataRegistry/Source/DataRegistry/Private/DataRegistry.cpp b/Engine/Plugins/Runtime/DataRegistry/Source/DataRegistry/Private/DataRegistry.cpp
--- a/Engine/Plugins/Runtime/DataRegistry/Source/DataRegistry/Private/DataRegistry.cpp
+++ b/Engine/Plugins/Runtime/DataRegistry/Source/DataRegistry/Private/DataRegistry.cpp
@@ -991,18 +991,7 @@ float UDataRegistry::GetCurrentTime()
 	return (float)(FApp::GetCurrentTime() - GStartTime);
 }
 
-// this is a workaround for a bug in clang that applies musttail attribute to Source->TimerUpdate call in some cases
-#if defined(__has_attribute)
-#	if __has_attribute(no_profile_instrument_function)
-#		define NO_PROFILE_ATTRIBUTE __attribute__((no_profile_instrument_function))
-#	endif
-#endif
-
-#ifndef NO_PROFILE_ATTRIBUTE
-#	define NO_PROFILE_ATTRIBUTE
-#endif
-
-NO_PROFILE_ATTRIBUTE
+UE_NO_PROFILE_ATTRIBUTE	// this is a workaround for a bug in clang that applies musttail attribute to Source->TimerUpdate call in some cases
 void UDataRegistry::TimerUpdate()
 {
 	if (!bIsInitialized)
diff --git a/Engine/Source/Runtime/Core/Public/Clang/ClangPlatform.h b/Engine/Source/Runtime/Core/Public/Clang/ClangPlatform.h
--- a/Engine/Source/Runtime/Core/Public/Clang/ClangPlatform.h
+++ b/Engine/Source/Runtime/Core/Public/Clang/ClangPlatform.h
@@ -47,4 +47,14 @@
 
 #define FUNCTION_NON_NULL_RETURN_START [[gnu::returns_nonnull]]
 
+#if defined(__has_attribute)
+#	if __has_attribute(no_profile_instrument_function)
+#		define UE_NO_PROFILE_ATTRIBUTE __attribute__((no_profile_instrument_function))
+#	endif
+#endif
+
+#ifndef UE_NO_PROFILE_ATTRIBUTE
+#	define UE_NO_PROFILE_ATTRIBUTE
+#endif
+
 #endif//defined(__clang__)
diff --git a/Engine/Source/Runtime/Core/Public/MSVC/MSVCPlatform.h b/Engine/Source/Runtime/Core/Public/MSVC/MSVCPlatform.h
--- a/Engine/Source/Runtime/Core/Public/MSVC/MSVCPlatform.h
+++ b/Engine/Source/Runtime/Core/Public/MSVC/MSVCPlatform.h
@@ -32,3 +32,5 @@
 		#define UE_LIFETIMEBOUND [[msvc::lifetimebound]]
 	#endif
 #endif
+
+#define UE_NO_PROFILE_ATTRIBUTE
\ No newline at end of file
