UE: Add data registry preload delegate. Useful for starting some async work before spending the synchronous game thread time to load data registries

diff --git a/Engine/Plugins/Runtime/DataRegistry/Source/DataRegistry/Private/DataRegistrySubsystem.cpp b/Engine/Plugins/Runtime/DataRegistry/Source/DataRegistry/Private/DataRegistrySubsystem.cpp
--- a/Engine/Plugins/Runtime/DataRegistry/Source/DataRegistry/Private/DataRegistrySubsystem.cpp
+++ b/Engine/Plugins/Runtime/DataRegistry/Source/DataRegistry/Private/DataRegistrySubsystem.cpp
@@ -328,6 +328,8 @@ void UDataRegistrySubsystem::LoadAllRegistries()
 		return;
 	}
 
+	OnPreLoadAllDataRegistries().Broadcast();
+
 	UAssetManager& AssetManager = UAssetManager::Get();
 	IAssetRegistry& AssetRegistry = AssetManager.GetAssetRegistry();
 
@@ -421,6 +423,11 @@ FDataRegistrySubsystemInitializedCallback& UDataRegistrySubsystem::OnSubsystemIn
 	return OnSubsystemInitializedCallback;
 }
 
+FPreLoadAllDataRegistriesCallback& UDataRegistrySubsystem::OnPreLoadAllDataRegistries()
+{
+	return OnPreLoadAllDataRegistriesCallback;
+}
+
 void UDataRegistrySubsystem::InitializeAllRegistries(bool bResetIfInitialized)
 {
 	for (const FRegistryMapPair& RegistryPair : RegistryMap)
diff --git a/Engine/Plugins/Runtime/DataRegistry/Source/DataRegistry/Public/DataRegistrySubsystem.h b/Engine/Plugins/Runtime/DataRegistry/Source/DataRegistry/Public/DataRegistrySubsystem.h
--- a/Engine/Plugins/Runtime/DataRegistry/Source/DataRegistry/Public/DataRegistrySubsystem.h
+++ b/Engine/Plugins/Runtime/DataRegistry/Source/DataRegistry/Public/DataRegistrySubsystem.h
@@ -167,6 +167,9 @@ public:
 	/** Accesses the delegate called when the subsystem has finished scanning for and initializing all known data registries */
 	FDataRegistrySubsystemInitializedCallback& OnSubsystemInitialized();
 
+	/** Accesses the delegate called before the subsystem has loaded data registries */
+	FPreLoadAllDataRegistriesCallback& OnPreLoadAllDataRegistries();
+
 	/** Initializes all loaded registries and prepares them for queries */
 	void InitializeAllRegistries(bool bResetIfInitialized = false);
 
@@ -274,6 +277,9 @@ protected:
 	/** Callback for when the subsystem has finished scanning for and initializing all known data registries */
 	FDataRegistrySubsystemInitializedCallback OnSubsystemInitializedCallback;
 
+	/** Callback for before the subsystem has loaded data registries */
+	FPreLoadAllDataRegistriesCallback OnPreLoadAllDataRegistriesCallback;
+
 	/** Singleton object for the DataRegistrySubsystem::Get function to use, populated in Initialize and cleared out in Deinitialize */
 	static TObjectPtr<UDataRegistrySubsystem> SingletonSubSystem;
 
diff --git a/Engine/Plugins/Runtime/DataRegistry/Source/DataRegistry/Public/DataRegistryTypes.h b/Engine/Plugins/Runtime/DataRegistry/Source/DataRegistry/Public/DataRegistryTypes.h
--- a/Engine/Plugins/Runtime/DataRegistry/Source/DataRegistry/Public/DataRegistryTypes.h
+++ b/Engine/Plugins/Runtime/DataRegistry/Source/DataRegistry/Public/DataRegistryTypes.h
@@ -539,5 +539,8 @@ DECLARE_MULTICAST_DELEGATE_OneParam(FDataRegistryCacheVersionCallback, class UDa
 /** Multicast delegate broadcast called when the data registry subsystem has finished scanning for and initializing all known data registries */
 DECLARE_MULTICAST_DELEGATE(FDataRegistrySubsystemInitializedCallback);
 
+/** Multicast delegate broadcast called before the data registry subsystem loads all registries */
+DECLARE_MULTICAST_DELEGATE(FPreLoadAllDataRegistriesCallback);
+
 
 DATAREGISTRY_API DECLARE_LOG_CATEGORY_EXTERN(LogDataRegistry, Log, All);
