Add GetItemName helper for DataRegistryOrTable

diff --git a/Engine/Plugins/Runtime/DataRegistry/Source/DataRegistry/Private/SoftDataRegistryOrTable.cpp b/Engine/Plugins/Runtime/DataRegistry/Source/DataRegistry/Private/SoftDataRegistryOrTable.cpp
--- a/Engine/Plugins/Runtime/DataRegistry/Source/DataRegistry/Private/SoftDataRegistryOrTable.cpp
+++ b/Engine/Plugins/Runtime/DataRegistry/Source/DataRegistry/Private/SoftDataRegistryOrTable.cpp
@@ -235,6 +235,16 @@ FString FDataRegistryOrTableRow::ToString() const
 	return DataTableRow.ToDebugString();
 }
 
+FName FDataRegistryOrTableRow::GetItemName() const
+{
+	if (bUseDataRegistryId)
+	{
+		return DataRegistryId.ItemName;
+	}
+
+	return DataTableRow.RowName;
+}
+
 const UDataRegistry* FDataRegistryOrTableRow::GetDataRegistry() const
 {
 	if (!bUseDataRegistryId)
diff --git a/Engine/Plugins/Runtime/DataRegistry/Source/DataRegistry/Public/SoftDataRegistryOrTable.h b/Engine/Plugins/Runtime/DataRegistry/Source/DataRegistry/Public/SoftDataRegistryOrTable.h
--- a/Engine/Plugins/Runtime/DataRegistry/Source/DataRegistry/Public/SoftDataRegistryOrTable.h
+++ b/Engine/Plugins/Runtime/DataRegistry/Source/DataRegistry/Public/SoftDataRegistryOrTable.h
@@ -225,6 +225,8 @@ struct FDataRegistryOrTableRow
 
 	DATAREGISTRY_API FString ToString() const;
 
+	DATAREGISTRY_API FName GetItemName() const;
+
 	// method to check validity of this row
 	DATAREGISTRY_API bool IsValid() const;
 
