Added WantsDelayedDataRegistryLoadingUntilPIE() method for checking whether or not Data Registries are expected to be lazy loaded in-editor.
Added IsReadyForInitialization() method for checking whether or not UDataRegistrySubsystem initialization "could" be started, with or without lazy loading. This will never return true until after the initial Asset Registry scan has completed.
Fixed a scenario where lazy loading of Data Registries wasn't being properly accounted for.

diff --git a/Engine/Plugins/Runtime/DataRegistry/Source/DataRegistry/Private/DataRegistrySubsystem.cpp b/Engine/Plugins/Runtime/DataRegistry/Source/DataRegistry/Private/DataRegistrySubsystem.cpp
--- a/Engine/Plugins/Runtime/DataRegistry/Source/DataRegistry/Private/DataRegistrySubsystem.cpp
+++ b/Engine/Plugins/Runtime/DataRegistry/Source/DataRegistry/Private/DataRegistrySubsystem.cpp
@@ -418,6 +418,15 @@ bool UDataRegistrySubsystem::IsConfigEnabled(bool bWarnIfNotEnabled /*= false*/)
 	return true;
 }
 
+bool UDataRegistrySubsystem::IsReadyForInitialization() const
+{
+#if WITH_EDITOR
+	return bReadyForInitialization || bLoadAllRegistriesOnNextPIE;
+#else
+	return bReadyForInitialization;
+#endif
+}
+
 FDataRegistrySubsystemInitializedCallback& UDataRegistrySubsystem::OnSubsystemInitialized()
 {
 	return OnSubsystemInitializedCallback;
@@ -952,7 +961,7 @@ void UDataRegistrySubsystem::PostAssetManager()
 	bReadyForInitialization = true;
 
 #if WITH_EDITOR
-	if (GIsEditor && !IsRunningCommandlet() && GetDefault<UDataRegistrySettings>()->bDelayLoadingDataRegistriesUntilPIE)
+	if (WantsDelayedDataRegistryLoadingUntilPIE())
 	{
 		bReadyForInitialization = false;
 		bLoadAllRegistriesOnNextPIE = true;
@@ -1026,6 +1035,11 @@ void UDataRegistrySubsystem::AddReferencedObjects(UObject* InThis, FReferenceCol
 
 #if WITH_EDITOR
 
+bool UDataRegistrySubsystem::WantsDelayedDataRegistryLoadingUntilPIE() const
+{
+	return GIsEditor && !IsRunningCommandlet() && GetDefault<UDataRegistrySettings>()->bDelayLoadingDataRegistriesUntilPIE;
+}
+
 void UDataRegistrySubsystem::PreBeginPIE(bool bStartSimulate)
 {
 	FlushDeferredRegistryLoad(true);
diff --git a/Engine/Plugins/Runtime/DataRegistry/Source/DataRegistry/Public/DataRegistrySubsystem.h b/Engine/Plugins/Runtime/DataRegistry/Source/DataRegistry/Public/DataRegistrySubsystem.h
--- a/Engine/Plugins/Runtime/DataRegistry/Source/DataRegistry/Public/DataRegistrySubsystem.h
+++ b/Engine/Plugins/Runtime/DataRegistry/Source/DataRegistry/Public/DataRegistrySubsystem.h
@@ -166,6 +166,9 @@ public:
 	/** Returns true if the system is enabled via any config scan settings, will optionally warn if not enabled */
 	UE_API bool IsConfigEnabled(bool bWarnIfNotEnabled = false) const;
 
+	/** Returns true if the system is ready for initialization, if it hasn't started already */
+	UE_API bool IsReadyForInitialization() const;
+
 	/** Accesses the delegate called when the subsystem has finished scanning for and initializing all known data registries */
 	UE_API FDataRegistrySubsystemInitializedCallback& OnSubsystemInitialized();
 
@@ -242,6 +245,11 @@ public:
 	/** Dumps out a text representation of every item in the registry */
 	static UE_API void DumpCachedItems(const TArray<FString>& Args);
 
+#if WITH_EDITOR
+	/** Returns true if the system will wait until first access to process pending Data Registry loads */
+	UE_API bool WantsDelayedDataRegistryLoadingUntilPIE() const;
+#endif
+
 protected:
 	typedef TPair<FName, TObjectPtr<UDataRegistry>> FRegistryMapPair;
 	TSortedMap<FName, TObjectPtr<UDataRegistry>, FDefaultAllocator, FNameFastLess> RegistryMap;
