Add warning logging for data registry reset scalable float variable
- example log: "LogDataRegistryIdCustomization: Warning: CustomContext:Default.MyDataRegistry.Registry1 reset to None because it's not a valid data registry ID"

#changelist validated

diff --git a/Engine/Plugins/Runtime/DataRegistry/Source/DataRegistryEditor/Private/DataRegistryIdCustomization.cpp b/Engine/Plugins/Runtime/DataRegistry/Source/DataRegistryEditor/Private/DataRegistryIdCustomization.cpp
--- a/Engine/Plugins/Runtime/DataRegistry/Source/DataRegistryEditor/Private/DataRegistryIdCustomization.cpp
+++ b/Engine/Plugins/Runtime/DataRegistry/Source/DataRegistryEditor/Private/DataRegistryIdCustomization.cpp
@@ -12,6 +12,8 @@
 
 #include UE_INLINE_GENERATED_CPP_BY_NAME(DataRegistryIdCustomization)
 
+DEFINE_LOG_CATEGORY_STATIC(LogDataRegistryIdCustomization, Log, All);
+
 #define LOCTEXT_NAMESPACE "DataRegistryEditor"
 
 FText FDataRegistryIdEditWrapper::GetPreviewDescription() const
@@ -148,6 +150,7 @@ void SDataRegistryItemNameWidget::OnTypeChanged(bool bClearInvalid)
 
 	if (!CachedIdValue.ItemName.IsNone() && !CachedIds.Contains(CachedIdValue) && bClearInvalid)
 	{
+		UE_LOG(LogDataRegistryIdCustomization, Warning, TEXT("%s reset to None because it's not a valid data registry ID"), *CachedIdValue.ToString());
 		// Name no longer valid, clear
 		OnSetId.Execute(FDataRegistryId(CachedIdValue.RegistryType, NAME_None));
 	}
