Uniformize DataRegistryId editor customization with DateRegistryType one in order to get the button to open the registry asset

#FailedValidations Preflight

diff --git a/Engine/Plugins/Runtime/DataRegistry/Source/DataRegistryEditor/Private/DataRegistryIdCustomization.cpp b/Engine/Plugins/Runtime/DataRegistry/Source/DataRegistryEditor/Private/DataRegistryIdCustomization.cpp
--- a/Engine/Plugins/Runtime/DataRegistry/Source/DataRegistryEditor/Private/DataRegistryIdCustomization.cpp
+++ b/Engine/Plugins/Runtime/DataRegistry/Source/DataRegistryEditor/Private/DataRegistryIdCustomization.cpp
@@ -232,6 +232,24 @@ EVisibility SDataRegistryItemNameWidget::GetWarningVisibility() const
 	return EVisibility::Collapsed;
 }
 
+namespace DataRegistryCustomizationPrivate
+{
+	template <typename T>
+	T* GetPropertyData(TSharedPtr<IPropertyHandle> StructPropertyHandle)
+	{
+		if (StructPropertyHandle.IsValid())
+		{
+			void* StructAddress = nullptr;
+			if (StructPropertyHandle->GetValueData(StructAddress) == FPropertyAccess::Success)
+			{
+				return static_cast<T*>(StructAddress);
+			}
+		}
+
+		return nullptr;
+	}
+}
+
 TSharedRef<IPropertyTypeCustomization> FDataRegistryIdCustomization::MakeInstance()
 {
 	return MakeShareable(new FDataRegistryIdCustomization);
@@ -253,36 +271,12 @@ void FDataRegistryIdCustomization::CustomizeHeader(TSharedRef<class IPropertyHan
 void FDataRegistryIdCustomization::CustomizeChildren(TSharedRef<class IPropertyHandle> InStructPropertyHandle, class IDetailChildrenBuilder& StructBuilder, IPropertyTypeCustomizationUtils& StructCustomizationUtils)
 {
 	/** Get all the existing property handles, need the name inside the type */
-	TypePropertyHandle = InStructPropertyHandle->GetChildHandle(GET_MEMBER_NAME_CHECKED(FDataRegistryId, RegistryType))->GetChildHandle(FName("Name"));
+	TypePropertyHandle = InStructPropertyHandle->GetChildHandle(GET_MEMBER_NAME_CHECKED(FDataRegistryId, RegistryType));
 	NamePropertyHandle = InStructPropertyHandle->GetChildHandle(GET_MEMBER_NAME_CHECKED(FDataRegistryId, ItemName));
 
 	if (TypePropertyHandle->IsValidHandle() && NamePropertyHandle->IsValidHandle())
 	{
-		FName FilterStructName;
-		if (StructPropertyHandle->HasMetaData(FDataRegistryType::ItemStructMetaData))
-		{
-			const FString& RowType = StructPropertyHandle->GetMetaData(FDataRegistryType::ItemStructMetaData);
-			FilterStructName = FName(*RowType);
-		}
-
-		FPropertyComboBoxArgs TypeArgs(TypePropertyHandle,
-			FOnGetPropertyComboBoxStrings::CreateStatic(FDataRegistryEditorModule::GenerateDataRegistryTypeComboBoxStrings, true, FilterStructName),
-			FOnGetPropertyComboBoxValue::CreateSP(this, &FDataRegistryIdCustomization::OnGetTypeValueString));
-		TypeArgs.ShowSearchForItemCount = 1;
-
-		/** Construct a combo box widget to select type */
-		StructBuilder.AddCustomRow(LOCTEXT("RegistryType", "Registry Type"))
-			.NameContent()
-			[
-				SNew(STextBlock)
-				.Text(LOCTEXT("RegistryType", "Registry Type"))
-				.Font(StructCustomizationUtils.GetRegularFont())
-			]
-			.ValueContent()
-			.MaxDesiredWidth(0.0f)
-			[
-				PropertyCustomizationHelpers::MakePropertyComboBox(TypeArgs)
-			];
+		StructBuilder.AddProperty(TypePropertyHandle.ToSharedRef());
 
 		/** Construct a combo box widget to select name */
 		StructBuilder.AddCustomRow(LOCTEXT("ItemName", "Item Name"))
@@ -307,25 +301,29 @@ void FDataRegistryIdCustomization::CustomizeChildren(TSharedRef<class IPropertyH
 FDataRegistryId FDataRegistryIdCustomization::GetCurrentValue() const
 {
 	FDataRegistryId OutId;
-	if (TypePropertyHandle.IsValid() && TypePropertyHandle->IsValidHandle() && NamePropertyHandle.IsValid() && NamePropertyHandle->IsValidHandle())
+	const FDataRegistryType* DataRegistryType = DataRegistryCustomizationPrivate::GetPropertyData<FDataRegistryType>(TypePropertyHandle);
+	if (DataRegistryType && NamePropertyHandle.IsValid() && NamePropertyHandle->IsValidHandle())
 	{
-		TypePropertyHandle->GetValue(OutId.RegistryType);
+		OutId.RegistryType = *DataRegistryType;
 		NamePropertyHandle->GetValue(OutId.ItemName);
 	}
-	
+
 	return OutId;
 }
 
 void FDataRegistryIdCustomization::SetCurrentValue(FDataRegistryId NewId)
 {
-	FDataRegistryId CurrentId = GetCurrentValue();
+	const FDataRegistryId CurrentId = GetCurrentValue();
 
 	// Always set value, only set type if different
 	NamePropertyHandle->SetValue(NewId.ItemName);
 
 	if (NewId.RegistryType != CurrentId.RegistryType)
 	{
-		TypePropertyHandle->SetValue(NewId.RegistryType.GetName());
+		if (FDataRegistryType* DataRegistryType = DataRegistryCustomizationPrivate::GetPropertyData<FDataRegistryType>(TypePropertyHandle))
+		{
+			*DataRegistryType = NewId.RegistryType;
+		}
 	}
 }
 
@@ -341,9 +339,9 @@ FString FDataRegistryIdCustomization::OnGetTypeValueString() const
 	{
 		return CurrentId.RegistryType.ToString();
 	}
-	
-	FName TempName;
-	if (TypePropertyHandle.IsValid() && TypePropertyHandle->GetValue(TempName) == FPropertyAccess::MultipleValues)
+
+	FDataRegistryType TempRegistry;
+	if (TypePropertyHandle.IsValid() && TypePropertyHandle->GetValue(TempRegistry) == FPropertyAccess::MultipleValues)
 	{
 		return LOCTEXT("MultipleValues", "Multiple Values").ToString();
 	}
