Added a function to get the class struct from FSoftDataRegistryOrTable

diff --git a/Engine/Plugins/Runtime/DataRegistry/Source/DataRegistry/Private/SoftDataRegistryOrTable.cpp b/Engine/Plugins/Runtime/DataRegistry/Source/DataRegistry/Private/SoftDataRegistryOrTable.cpp
--- a/Engine/Plugins/Runtime/DataRegistry/Source/DataRegistry/Private/SoftDataRegistryOrTable.cpp
+++ b/Engine/Plugins/Runtime/DataRegistry/Source/DataRegistry/Private/SoftDataRegistryOrTable.cpp
@@ -167,6 +167,28 @@ const UDataRegistry* FSoftDataRegistryOrTable::GetDataRegistry() const
 	return RegistrySystem->GetRegistryForType(RegistryType);
 }
 
+const UScriptStruct* FSoftDataRegistryOrTable::GetStruct() const
+{
+	if (bUseDataRegistry)
+	{
+		const UDataRegistry* EventRegistry = GetDataRegistry();
+		if (!EventRegistry)
+		{
+			UE_LOG(LogDataRegistry, Warning, TEXT("[%hs] No Registry found  Registry:%s"), __FUNCTION__, *RegistryType.GetName().ToString());
+			return nullptr;
+		}
+
+		return EventRegistry->GetItemStruct();
+	}
+
+	if (UDataTable* TablePtr = Table.LoadSynchronous())
+	{
+		return TablePtr->GetRowStruct();
+	}
+
+	return nullptr;
+}
+
 bool FDataRegistryOrTableRow::SerializeFromMismatchedTag(const FPropertyTag& Tag, FStructuredArchive::FSlot Slot)
 {
 	static const FLazyName DataTableRowHandleName("DataTableRowHandle");
diff --git a/Engine/Plugins/Runtime/DataRegistry/Source/DataRegistry/Public/SoftDataRegistryOrTable.h b/Engine/Plugins/Runtime/DataRegistry/Source/DataRegistry/Public/SoftDataRegistryOrTable.h
--- a/Engine/Plugins/Runtime/DataRegistry/Source/DataRegistry/Public/SoftDataRegistryOrTable.h
+++ b/Engine/Plugins/Runtime/DataRegistry/Source/DataRegistry/Public/SoftDataRegistryOrTable.h
@@ -152,6 +152,11 @@ struct FSoftDataRegistryOrTable
 	 */
 	DATAREGISTRY_API void LoadAsync(FStreamableDelegate DelegateToCall);	
 
+	/**
+	 * Returns the script struct used by this data registry or table
+	 */
+	DATAREGISTRY_API const UScriptStruct* GetStruct() const;
+
 	/** Method to get a FDataRegistryOrTableRow from the given row name */
 	DATAREGISTRY_API FDataRegistryOrTableRow GetRegistryOrTableRow(FName RowName) const;
 
