Added FDataTableOrRegistry and FSoftDataRegistryOrTable editor row customization, to make it more compact in the details view

diff --git a/Engine/Plugins/Runtime/DataRegistry/Source/DataRegistryEditor/Private/DataRegistryEditorModule.cpp b/Engine/Plugins/Runtime/DataRegistry/Source/DataRegistryEditor/Private/DataRegistryEditorModule.cpp
--- a/Engine/Plugins/Runtime/DataRegistry/Source/DataRegistryEditor/Private/DataRegistryEditorModule.cpp
+++ b/Engine/Plugins/Runtime/DataRegistry/Source/DataRegistryEditor/Private/DataRegistryEditorModule.cpp
@@ -3,6 +3,8 @@
 #include "DataRegistryEditorModule.h"
 #include "DataRegistryIdCustomization.h"
 #include "DataRegistryTypeCustomization.h"
+#include "DataTableOrRegistryRowCustomization.h"
+#include "SoftDataRegistryOrTableCustomization.h"
 #include "DataRegistrySubsystem.h"
 #include "AssetTypeActions_DataRegistry.h"
 #include "EdGraphSchema_K2.h"
@@ -46,6 +48,8 @@ void FDataRegistryEditorModule::StartupModule()
 		FPropertyEditorModule& PropertyModule = FModuleManager::LoadModuleChecked<FPropertyEditorModule>("PropertyEditor");
 		PropertyModule.RegisterCustomPropertyTypeLayout("DataRegistryId", FOnGetPropertyTypeCustomizationInstance::CreateStatic(&FDataRegistryIdCustomization::MakeInstance));
 		PropertyModule.RegisterCustomPropertyTypeLayout("DataRegistryType", FOnGetPropertyTypeCustomizationInstance::CreateStatic(&FDataRegistryTypeCustomization::MakeInstance));
+		PropertyModule.RegisterCustomPropertyTypeLayout("DataRegistryOrTableRow", FOnGetPropertyTypeCustomizationInstance::CreateStatic(&FDataRegistryOrTableRowCustomization::MakeInstance));
+		PropertyModule.RegisterCustomPropertyTypeLayout("SoftDataRegistryOrTable", FOnGetPropertyTypeCustomizationInstance::CreateStatic(&FSoftDataRegistryOrTableCustomization::MakeInstance));
 		PropertyModule.NotifyCustomizationModuleChanged();
 
 		// Register Pins and nodes
@@ -68,6 +72,16 @@ void FDataRegistryEditorModule::ShutdownModule()
 		{
 			FEdGraphUtilities::UnregisterVisualPinFactory(DataRegistryGraphPanelPinFactory);
 		}
+
+		if (FModuleManager::Get().IsModuleLoaded("PropertyEditor"))
+		{
+			FPropertyEditorModule& PropertyModule = FModuleManager::LoadModuleChecked<FPropertyEditorModule>("PropertyEditor");
+
+			PropertyModule.UnregisterCustomPropertyTypeLayout("SoftDataRegistryOrTable");
+			PropertyModule.UnregisterCustomPropertyTypeLayout("DataRegistryOrTableRow");
+			PropertyModule.UnregisterCustomPropertyTypeLayout("DataRegistryType");
+			PropertyModule.UnregisterCustomPropertyTypeLayout("DataRegistryId");
+		}
 	}
 
 	MenuExtensibilityManager.Reset();
diff --git a/Engine/Plugins/Runtime/DataRegistry/Source/DataRegistryEditor/Private/DataTableOrRegistryRowCustomization.cpp b/Engine/Plugins/Runtime/DataRegistry/Source/DataRegistryEditor/Private/DataTableOrRegistryRowCustomization.cpp
new file mode 100644
--- /dev/null
+++ b/Engine/Plugins/Runtime/DataRegistry/Source/DataRegistryEditor/Private/DataTableOrRegistryRowCustomization.cpp
@@ -0,0 +1,555 @@
+// Copyright Epic Games, Inc. All Rights Reserved.
+
+#include "DataTableOrRegistryRowCustomization.h"
+#include "Components/WrapBox.h"
+#include "DetailWidgetRow.h"
+#include "PropertyCustomizationHelpers.h"
+#include "DataRegistryEditorModule.h"
+#include "IDetailChildrenBuilder.h"
+#include "DataRegistryIdCustomization.h"
+#include "DataRegistrySubsystem.h"
+#include "PropertyCustomizationHelpers.h"
+#include "Widgets/Images/SImage.h"
+#include "Widgets/Input/SButton.h"
+
+#define LOCTEXT_NAMESPACE "DataRegistryOrTableRow"
+
+
+// Makes a new instance of this detail layout class for a specific detail view requesting it
+TSharedRef<IPropertyTypeCustomization> FDataRegistryOrTableRowCustomization::MakeInstance()
+{
+	return MakeShareable(new FDataRegistryOrTableRowCustomization);
+}
+
+void FDataRegistryOrTableRowCustomization::CustomizeHeader(TSharedRef<class IPropertyHandle> InStructPropertyHandle, class FDetailWidgetRow& HeaderRow, IPropertyTypeCustomizationUtils& StructCustomizationUtils)
+{
+	this->StructPropertyHandle = InStructPropertyHandle;
+	
+	FilterStructName = NAME_None;
+	if (StructPropertyHandle->HasMetaData(FDataRegistryType::ItemStructMetaData))
+	{
+		const FString& RowType = StructPropertyHandle->GetMetaData(FDataRegistryType::ItemStructMetaData);
+		FilterStructName = FName(*RowType);
+	}
+}
+
+void FDataRegistryOrTableRowCustomization::GetDataSourceComboStrings(TArray< TSharedPtr<FString> >& OutComboBoxStrings, TArray<TSharedPtr<SToolTip>>& OutToolTips, TArray<bool>& OutRestrictedItems) const
+{
+	OutComboBoxStrings.Add(MakeShared<FString>(LOCTEXT("DataTable", "Data Table").ToString()));
+	OutComboBoxStrings.Add(MakeShared<FString>(LOCTEXT("DataRegistry", "Data Registry").ToString()));
+	OutRestrictedItems.Add(false);
+	OutRestrictedItems.Add(false);	
+}
+
+FString FDataRegistryOrTableRowCustomization::GetDataSourceTypeValueString() const
+{
+	if (!UseRegistryPropertyHandle.IsValid() || !UseRegistryPropertyHandle->IsValidHandle())
+	{
+		return LOCTEXT("NoneValue", "None").ToString();	
+	}
+
+	bool bUseDataRegistryId = false;
+	const FPropertyAccess::Result SourceResult = UseRegistryPropertyHandle->GetValue(bUseDataRegistryId);
+	if (SourceResult == FPropertyAccess::Success)
+	{
+		if (bUseDataRegistryId)
+		{
+			return LOCTEXT("DataRegistry", "Data Registry").ToString();
+		}
+		else
+		{
+			return LOCTEXT("DataTable", "Data Table").ToString();	
+		}
+	}
+	else if (SourceResult == FPropertyAccess::Fail)
+	{
+		return LOCTEXT("DataTable_None", "None").ToString();
+	}
+	else
+	{
+		return LOCTEXT("MultipleValues", "Multiple Values").ToString();
+	}
+}
+
+
+void FDataRegistryOrTableRowCustomization::OnDataSourceSelected(const FString& String)
+{
+	const bool bUsingDataRegistry = String == LOCTEXT("DataRegistry", "Data Registry").ToString();
+
+	static auto CheckHandle = [](const TSharedPtr<IPropertyHandle>& Handle)
+	{
+		return Handle.IsValid() && Handle->IsValidHandle();
+	};
+
+	if (CheckHandle(UseRegistryPropertyHandle))
+	{
+		UseRegistryPropertyHandle->SetValue(bUsingDataRegistry);
+	}
+
+	if (bUsingDataRegistry)
+	{
+		// clear the data table values to prevent hidden references to a resource
+		if (CheckHandle(TablePropertyHandle) && CheckHandle(TableRowNamePropertyHandle))
+		{
+			TablePropertyHandle->SetValue(static_cast<UDataTable*>(nullptr));
+			TableRowNamePropertyHandle->SetValue(NAME_None);
+		}
+	}
+	else
+	{
+		// clear the data registry values to prevent hidden references to the registry
+		if (CheckHandle(RegistryTypeNamePropertyHandle) && CheckHandle(RegistryItemNamePropertyHandle))
+		{
+			RegistryTypeNamePropertyHandle->SetValue(FDataRegistryType().GetName());
+			RegistryItemNamePropertyHandle->SetValue(NAME_None);
+		}
+	}
+}
+
+FDataRegistryId FDataRegistryOrTableRowCustomization::GetCurrentDataRegistryValue() const
+{
+	FDataRegistryId Result {};
+
+	static auto CheckHandle = [](const TSharedPtr<IPropertyHandle>& Handle)
+	{
+		return Handle.IsValid() && Handle->IsValidHandle();
+	};
+
+	if (CheckHandle(RegistryTypeNamePropertyHandle) && CheckHandle(RegistryItemNamePropertyHandle))
+	{
+		RegistryTypeNamePropertyHandle->GetValue(static_cast<FName&>(Result.RegistryType));
+		RegistryItemNamePropertyHandle->GetValue(Result.ItemName);		
+	}
+
+	return Result;
+
+}
+
+void FDataRegistryOrTableRowCustomization::SetCurrentDataRegistryValue(FDataRegistryId NewValue)
+{
+	static auto CheckHandle = [](const TSharedPtr<IPropertyHandle>& Handle)
+	{
+		return Handle.IsValid() && Handle->IsValidHandle();
+	};
+
+	if (CheckHandle(RegistryTypeNamePropertyHandle) && CheckHandle(RegistryItemNamePropertyHandle))
+	{
+		RegistryTypeNamePropertyHandle->SetValue(NewValue.RegistryType.GetName());
+		RegistryItemNamePropertyHandle->SetValue(NewValue.ItemName);		
+	}
+}
+
+FDataRegistryOrTableRow FDataRegistryOrTableRowCustomization::GetCurrentValue() const
+{
+	FDataRegistryOrTableRow Result;
+
+	static auto CheckHandle = [](const TSharedPtr<IPropertyHandle>& Handle)
+	{
+		return Handle.IsValid() && Handle->IsValidHandle();
+	};
+
+	if (!CheckHandle(UseRegistryPropertyHandle) ||
+		!CheckHandle(RegistryTypeNamePropertyHandle) ||
+		!CheckHandle(RegistryItemNamePropertyHandle) ||
+		!CheckHandle(TablePropertyHandle) ||
+		!CheckHandle(TableRowNamePropertyHandle))
+	{
+		return Result;
+	}
+
+	UseRegistryPropertyHandle->GetValue(Result.bUseDataRegistryId);
+
+	if (Result.bUseDataRegistryId)
+	{
+		RegistryTypeNamePropertyHandle->GetValue(static_cast<FName&>(Result.DataRegistryId.RegistryType));
+		RegistryItemNamePropertyHandle->GetValue(Result.DataRegistryId.ItemName);
+	}
+	else
+	{
+		UObject* SourceDataTable = nullptr;
+
+		TablePropertyHandle->GetValue(SourceDataTable);
+		Result.DataTableRow.DataTable = Cast<UDataTable>(SourceDataTable);
+
+		TableRowNamePropertyHandle->GetValue(Result.DataTableRow.RowName);
+	}
+
+	return Result;
+}
+
+FString FDataRegistryOrTableRowCustomization::OnGetDataRegistryTypeValueString() const
+{
+	FDataRegistryOrTableRow CurrentId = GetCurrentValue();
+	if (CurrentId.bUseDataRegistryId)
+	{
+		return CurrentId.DataRegistryId.RegistryType.ToString();
+	}
+
+	return LOCTEXT("NoneValue", "None").ToString();
+}
+
+bool FDataRegistryOrTableRowCustomization::ShouldFilterDataTableAsset(const struct FAssetData& AssetData)
+{
+	if (!FilterStructName.IsNone())
+	{
+		static const FName RowStructureTagName("RowStructure");
+		UScriptStruct* RowFilterStruct = UClass::TryFindTypeSlow<UScriptStruct>(FilterStructName.ToString());
+		FString RowStructure;
+		if (AssetData.GetTagValue<FString>(RowStructureTagName, RowStructure))
+		{
+			// This is slow, but at the moment we don't have an alternative to the short struct name search
+			UScriptStruct* RowStruct = UClass::TryFindTypeSlow<UScriptStruct>(RowStructure);
+
+			if (RowStruct && RowFilterStruct && RowStruct->IsChildOf(RowFilterStruct))
+			{
+				return false;
+			}
+		}
+		return true;
+	}
+	return false;
+}
+
+void FDataRegistryOrTableRowCustomization::OnGetDataTableRowStrings(TArray< TSharedPtr<FString> >& OutStrings, TArray<TSharedPtr<SToolTip>>& OutToolTips, TArray<bool>& OutRestrictedItems) const
+{
+	FDataRegistryOrTableRow CurrentId = GetCurrentValue();
+	if (CurrentId.bUseDataRegistryId)
+	{
+		return;
+	}
+
+	TArray<FName> AllRowNames;
+	if (CurrentId.DataTableRow.DataTable != nullptr)
+	{
+		for (TMap<FName, uint8*>::TConstIterator Iterator(CurrentId.DataTableRow.DataTable->GetRowMap()); Iterator; ++Iterator)
+		{
+			AllRowNames.Add(Iterator.Key());
+		}
+
+		// Sort the names alphabetically.
+		AllRowNames.Sort(FNameLexicalLess());
+	}
+
+	for (const FName& RowName : AllRowNames)
+	{
+		OutStrings.Add(MakeShared<FString>(RowName.ToString()));
+		OutRestrictedItems.Add(false);
+	}
+}
+
+FString FDataRegistryOrTableRowCustomization::OnGetDataTableRowValueString() const
+{
+	if (!TableRowNamePropertyHandle.IsValid() || !TableRowNamePropertyHandle->IsValidHandle())
+	{
+		return FString();
+	}
+
+	FName RowNameValue;
+	const FPropertyAccess::Result RowResult = TableRowNamePropertyHandle->GetValue(RowNameValue);
+	if (RowResult == FPropertyAccess::Success)
+	{
+		if (RowNameValue.IsNone())
+		{
+			return LOCTEXT("DataTable_None", "None").ToString();
+		}
+		return RowNameValue.ToString();
+	}
+	else if (RowResult == FPropertyAccess::Fail)
+	{
+		return LOCTEXT("DataTable_None", "None").ToString();
+	}
+	else
+	{
+		return LOCTEXT("MultipleValues", "Multiple Values").ToString();
+	}
+}
+
+FText FDataRegistryOrTableRowCustomization::OnGetDataRegistryNameValueText() const
+{
+	FDataRegistryId CurrentId = GetCurrentDataRegistryValue();
+	if (!CurrentId.ItemName.IsNone())
+	{
+		return FText::FromName(CurrentId.ItemName);
+	}
+
+	FName TempName;
+	if (RegistryItemNamePropertyHandle.IsValid() && RegistryItemNamePropertyHandle->GetValue(TempName) == FPropertyAccess::MultipleValues)
+	{
+		return LOCTEXT("MultipleValues", "Multiple Values");
+	}
+
+	return LOCTEXT("NoneValue", "None");
+}
+
+EVisibility FDataRegistryOrTableRowCustomization::GetOpenDataRegistryAssetVisibility() const
+{
+	FDataRegistryId CurrentId = GetCurrentDataRegistryValue();
+	return CurrentId.RegistryType.IsValid() ? EVisibility::Visible : EVisibility::Collapsed;
+}
+
+FReply FDataRegistryOrTableRowCustomization::OnClickOpenDataRegistryAsset()
+{
+	FDataRegistryId CurrentId = GetCurrentDataRegistryValue();
+
+	if (const UDataRegistrySubsystem* DataRegistrySubsystem = UDataRegistrySubsystem::Get())
+	{
+		if (UDataRegistry* DataRegistry = DataRegistrySubsystem->GetRegistryForType(CurrentId.RegistryType))
+		{
+			GEditor->GetEditorSubsystem<UAssetEditorSubsystem>()->OpenEditorForAsset(DataRegistry);
+		}
+	}
+
+	return FReply::Handled();
+}
+
+
+FText FDataRegistryOrTableRowCustomization::GetOpenDataRegistryAssetTooltip() const
+{
+	FDataRegistryId CurrentId = GetCurrentDataRegistryValue();
+
+	if (const UDataRegistrySubsystem* DataRegistrySubsystem = UDataRegistrySubsystem::Get())
+	{
+		if (UDataRegistry* DataRegistry = DataRegistrySubsystem->GetRegistryForType(CurrentId.RegistryType))
+		{
+			FFormatNamedArguments Args;
+			Args.Add(TEXT("Asset"), FText::AsCultureInvariant(DataRegistry->GetName()));
+			return FText::Format(LOCTEXT("OpenSpecificDataRegistry", "Open '{Asset}' in the editor"), Args);
+		}
+	}
+
+	return LOCTEXT("OpenDataRegistry", "Open the Data Registry in the editor");
+}
+
+
+void FDataRegistryOrTableRowCustomization::CustomizeChildren(TSharedRef<class IPropertyHandle> InStructPropertyHandle, class IDetailChildrenBuilder& StructBuilder, IPropertyTypeCustomizationUtils& StructCustomizationUtils)
+{
+
+	static auto CheckHandle = [](const TSharedPtr<IPropertyHandle>& Handle)
+	{
+		return Handle.IsValid() && Handle->IsValidHandle();
+	};
+
+	// get the child property handles
+	UseRegistryPropertyHandle = InStructPropertyHandle->GetChildHandle(GET_MEMBER_NAME_CHECKED(FDataRegistryOrTableRow, bUseDataRegistryId));
+	if (!CheckHandle(UseRegistryPropertyHandle))
+	{
+		return;
+	}
+
+	TSharedPtr<IPropertyHandle> RegistryIdPropertyHandle = InStructPropertyHandle->GetChildHandle(GET_MEMBER_NAME_CHECKED(FDataRegistryOrTableRow, DataRegistryId));
+
+	if (!CheckHandle(RegistryIdPropertyHandle))
+	{
+		return;
+	}
+
+	TSharedPtr<IPropertyHandle> RegistryTypePropertyHandle = RegistryIdPropertyHandle->GetChildHandle(GET_MEMBER_NAME_CHECKED(FDataRegistryId, RegistryType));
+	if (!CheckHandle(RegistryTypePropertyHandle))
+	{
+		return;
+	}
+
+	RegistryTypeNamePropertyHandle = RegistryTypePropertyHandle->GetChildHandle(FName("Name"));
+	if (!CheckHandle(RegistryTypeNamePropertyHandle))
+	{
+		return;
+	}
+
+	RegistryItemNamePropertyHandle = RegistryIdPropertyHandle->GetChildHandle(GET_MEMBER_NAME_CHECKED(FDataRegistryId, ItemName));
+	if (!CheckHandle(RegistryItemNamePropertyHandle))
+	{
+		return;
+	}
+
+	TSharedPtr<IPropertyHandle> DataTableRowPropertyHandle = InStructPropertyHandle->GetChildHandle(GET_MEMBER_NAME_CHECKED(FDataRegistryOrTableRow, DataTableRow));
+	if (!CheckHandle(DataTableRowPropertyHandle))
+	{
+		return;
+	}
+
+	TablePropertyHandle = DataTableRowPropertyHandle->GetChildHandle(GET_MEMBER_NAME_CHECKED(FDataTableRowHandle, DataTable));
+	if (!CheckHandle(TablePropertyHandle))
+	{
+		return;
+	}
+
+	TableRowNamePropertyHandle = DataTableRowPropertyHandle->GetChildHandle(GET_MEMBER_NAME_CHECKED(FDataTableRowHandle, RowName));
+	if (!CheckHandle(TableRowNamePropertyHandle))
+	{
+		return;
+	}
+
+	FPropertyComboBoxArgs SelectSourceTypeArgs(UseRegistryPropertyHandle,
+		FOnGetPropertyComboBoxStrings::CreateSP(this, &FDataRegistryOrTableRowCustomization::GetDataSourceComboStrings),
+		FOnGetPropertyComboBoxValue::CreateSP(this, &FDataRegistryOrTableRowCustomization::GetDataSourceTypeValueString),
+		FOnPropertyComboBoxValueSelected::CreateSP(this, &FDataRegistryOrTableRowCustomization::OnDataSourceSelected)
+	);
+
+	FPropertyComboBoxArgs DataRegistryIdTypeArgs(RegistryTypeNamePropertyHandle,
+		FOnGetPropertyComboBoxStrings::CreateStatic(FDataRegistryEditorModule::GenerateDataRegistryTypeComboBoxStrings, true, FilterStructName),
+		FOnGetPropertyComboBoxValue::CreateSP(this, &FDataRegistryOrTableRowCustomization::OnGetDataRegistryTypeValueString));
+	DataRegistryIdTypeArgs.ShowSearchForItemCount = 1;
+
+	FPropertyComboBoxArgs DataTableRowStringsArgs(TableRowNamePropertyHandle, 
+		FOnGetPropertyComboBoxStrings::CreateSP(this, &FDataRegistryOrTableRowCustomization::OnGetDataTableRowStrings), 
+		FOnGetPropertyComboBoxValue::CreateSP(this, &FDataRegistryOrTableRowCustomization::OnGetDataTableRowValueString));
+	DataTableRowStringsArgs.ShowSearchForItemCount = 1;
+
+	auto IsUsingDataRegistry = [this]()
+	{
+		bool bResult = false;
+		if (UseRegistryPropertyHandle.IsValid() && UseRegistryPropertyHandle->IsValidHandle())
+		{
+			UseRegistryPropertyHandle->GetValue(bResult);
+		}
+		return bResult ? EVisibility::Visible : EVisibility::Collapsed;
+	};
+
+	auto IsUsingDataTable = [this]()
+	{
+		bool bResult = false;
+		if (UseRegistryPropertyHandle.IsValid() && UseRegistryPropertyHandle->IsValidHandle())
+		{
+			UseRegistryPropertyHandle->GetValue(bResult);
+		}
+		return bResult ? EVisibility::Collapsed : EVisibility::Visible;
+	};
+
+	auto IsUsingDataTableAndSourceSelected = [this, IsUsingDataTable]() {
+
+		if (IsUsingDataTable() == EVisibility::Collapsed) 
+		{
+			return EVisibility::Collapsed;
+		}
+
+		UObject* SourceDataTable = nullptr;
+		if (TablePropertyHandle.IsValid() && TablePropertyHandle->IsValidHandle())
+		{
+			TablePropertyHandle->GetValue(SourceDataTable);
+		}
+
+		return SourceDataTable != nullptr ? EVisibility::Visible : EVisibility::Collapsed;
+	};
+
+	auto IsUsingDataRegistryAndSourceSelected = [this, IsUsingDataRegistry]() {
+
+		if (IsUsingDataRegistry() == EVisibility::Collapsed) 
+		{
+			return EVisibility::Collapsed;
+		}
+
+		FName RegistryType;
+
+		RegistryTypeNamePropertyHandle->GetValue(RegistryType);
+
+		return RegistryType.IsNone() ? EVisibility::Collapsed : EVisibility::Visible;
+	};
+
+	FUIAction CopyAction;
+	FUIAction PasteAction;
+	InStructPropertyHandle->CreateDefaultPropertyCopyPasteActions(CopyAction, PasteAction);
+
+	StructBuilder.AddCustomRow(LOCTEXT("DataTableOrRegistry", "Data Table Or Registry"))
+		.CopyAction(CopyAction)
+		.PasteAction(PasteAction)
+		.NameContent()
+		.HAlign(HAlign_Fill)
+		[
+			InStructPropertyHandle->CreatePropertyNameWidget()
+		]
+		.ValueContent()
+		.MinDesiredWidth(800.f)
+		[
+			SNew(SWrapBox)
+			.UseAllottedSize(true)
+			.Orientation(Orient_Horizontal)
+			.HAlign(HAlign_Fill)			
+			+ SWrapBox::Slot()
+				.FillEmptySpace(false)
+				[
+					PropertyCustomizationHelpers::MakePropertyComboBox(SelectSourceTypeArgs)
+				]
+			+ SWrapBox::Slot()
+				.FillEmptySpace(true)
+				[
+					SNew(SObjectPropertyEntryBox)
+					.PropertyHandle(TablePropertyHandle)
+					.AllowedClass(UDataTable::StaticClass())
+					.OnShouldFilterAsset(this, &FDataRegistryOrTableRowCustomization::ShouldFilterDataTableAsset)
+					.Visibility(TAttribute<EVisibility>::Create(IsUsingDataTable))
+				]
+			+ SWrapBox::Slot()
+				.FillEmptySpace(true)
+				[
+					SNew(SHorizontalBox)
+					.Visibility(TAttribute<EVisibility>::Create(IsUsingDataRegistry))
+					+ SHorizontalBox::Slot()
+						.AutoWidth()
+						[
+							PropertyCustomizationHelpers::MakePropertyComboBox(DataRegistryIdTypeArgs)
+						]
+					+ SHorizontalBox::Slot()
+						.AutoWidth()
+						.HAlign(HAlign_Center)
+						.VAlign(VAlign_Center)
+						.Padding(1, 0)
+						[
+							SNew(SBox)
+								.HAlign(HAlign_Center)
+								.VAlign(VAlign_Center)
+								.WidthOverride(22.0f)
+								.HeightOverride(22.0f)
+								.ToolTipText(this, &FDataRegistryOrTableRowCustomization::GetOpenDataRegistryAssetTooltip)
+								[
+									SNew(SButton)
+										.ButtonStyle(FAppStyle::Get(), "SimpleButton")
+										.OnClicked(this, &FDataRegistryOrTableRowCustomization::OnClickOpenDataRegistryAsset)
+										.ContentPadding(0.0f)
+										.IsFocusable(false)
+										.Visibility(this, &FDataRegistryOrTableRowCustomization::GetOpenDataRegistryAssetVisibility)
+										[
+											SNew(SImage)
+												.Image(FAppStyle::GetBrush("SystemWideCommands.SummonOpenAssetDialog"))
+												.ColorAndOpacity(FSlateColor::UseForeground())
+										]
+								]
+						]
+				]
+			+ SWrapBox::Slot()
+				.FillEmptySpace(true)
+				[
+					SNew(SBox)
+					.Padding(0.0f)
+					.Visibility(TAttribute<EVisibility>::Create(IsUsingDataTableAndSourceSelected))
+					.VAlign(VAlign_Center)
+					[
+						PropertyCustomizationHelpers::MakePropertyComboBox(DataTableRowStringsArgs)
+					]
+				]
+			+ SWrapBox::Slot()
+				.FillEmptySpace(true)
+				[
+					SNew(SDataRegistryItemNameWidget)
+					.OnGetDisplayText(this, &FDataRegistryOrTableRowCustomization::OnGetDataRegistryNameValueText)
+					.OnGetId(this, &FDataRegistryOrTableRowCustomization::GetCurrentDataRegistryValue)
+					.OnSetId(this, &FDataRegistryOrTableRowCustomization::SetCurrentDataRegistryValue)
+					.bAllowClear(true)
+					.Visibility(TAttribute<EVisibility>::Create(IsUsingDataRegistryAndSourceSelected))
+				]
+		]
+		.OverrideResetToDefault(FResetToDefaultOverride::Create(
+			TAttribute<bool>::CreateLambda([InStructPropertyHandle]() -> bool
+			{
+				return InStructPropertyHandle->IsValidHandle() && InStructPropertyHandle->DiffersFromDefault();
+			}),
+			FSimpleDelegate::CreateLambda([InStructPropertyHandle]()
+			{
+				if (InStructPropertyHandle->IsValidHandle())
+				{
+					InStructPropertyHandle->ResetToDefault();
+				}
+			})
+		));
+}
+
+#undef LOCTEXT_NAMESPACE
diff --git a/Engine/Plugins/Runtime/DataRegistry/Source/DataRegistryEditor/Private/DataTableOrRegistryRowCustomization.h b/Engine/Plugins/Runtime/DataRegistry/Source/DataRegistryEditor/Private/DataTableOrRegistryRowCustomization.h
new file mode 100644
--- /dev/null
+++ b/Engine/Plugins/Runtime/DataRegistry/Source/DataRegistryEditor/Private/DataTableOrRegistryRowCustomization.h
@@ -0,0 +1,51 @@
+// Copyright Epic Games, Inc. All Rights Reserved.
+
+#pragma once
+
+#include "SoftDataRegistryOrTable.h"
+#include "IPropertyTypeCustomization.h"
+#include "Widgets/SToolTip.h"
+
+/**
+ * Customizes a FDataRegistryOrTableRow to support filtering by item struct, and compacting it to one row 
+ */
+class FDataRegistryOrTableRowCustomization : public IPropertyTypeCustomization
+{
+public:
+	// Makes a new instance of this detail layout class for a specific detail view requesting it
+	static TSharedRef<IPropertyTypeCustomization> MakeInstance();
+
+	/** IPropertyTypeCustomization interface */
+	virtual void CustomizeHeader(TSharedRef<class IPropertyHandle> InStructPropertyHandle, class FDetailWidgetRow& HeaderRow, IPropertyTypeCustomizationUtils& StructCustomizationUtils) override;
+	virtual void CustomizeChildren(TSharedRef<class IPropertyHandle> InStructPropertyHandle, class IDetailChildrenBuilder& StructBuilder, IPropertyTypeCustomizationUtils& StructCustomizationUtils) override;
+
+private:
+	void GetDataSourceComboStrings(TArray< TSharedPtr<FString> >& OutComboBoxStrings, TArray<TSharedPtr<SToolTip>>& OutToolTips, TArray<bool>& OutRestrictedItems) const;
+	FString GetDataSourceTypeValueString() const;
+	void OnDataSourceSelected(const FString& String);
+
+	FString OnGetDataRegistryTypeValueString() const;
+	bool ShouldFilterDataTableAsset(const struct FAssetData& AssetData);
+
+	void OnGetDataTableRowStrings(TArray< TSharedPtr<FString> >& OutStrings, TArray<TSharedPtr<SToolTip>>& OutToolTips, TArray<bool>& OutRestrictedItems) const;
+	FString OnGetDataTableRowValueString() const;
+
+	FDataRegistryOrTableRow GetCurrentValue() const;
+	EVisibility GetOpenDataRegistryAssetVisibility() const;
+	FReply OnClickOpenDataRegistryAsset();
+
+	FText GetOpenDataRegistryAssetTooltip() const;
+	FText OnGetDataRegistryNameValueText() const;
+	FDataRegistryId GetCurrentDataRegistryValue() const;
+	void SetCurrentDataRegistryValue(FDataRegistryId NewValue);
+
+	TSharedPtr<IPropertyHandle> StructPropertyHandle;
+
+	FName FilterStructName;
+
+	TSharedPtr<IPropertyHandle> UseRegistryPropertyHandle;
+	TSharedPtr<IPropertyHandle> RegistryTypeNamePropertyHandle;
+	TSharedPtr<IPropertyHandle> RegistryItemNamePropertyHandle;
+	TSharedPtr<IPropertyHandle> TablePropertyHandle;
+	TSharedPtr<IPropertyHandle> TableRowNamePropertyHandle;
+};
diff --git a/Engine/Plugins/Runtime/DataRegistry/Source/DataRegistryEditor/Private/SoftDataRegistryOrTableCustomization.cpp b/Engine/Plugins/Runtime/DataRegistry/Source/DataRegistryEditor/Private/SoftDataRegistryOrTableCustomization.cpp
new file mode 100644
--- /dev/null
+++ b/Engine/Plugins/Runtime/DataRegistry/Source/DataRegistryEditor/Private/SoftDataRegistryOrTableCustomization.cpp
@@ -0,0 +1,367 @@
+// Copyright Epic Games, Inc. All Rights Reserved.
+
+#include "SoftDataRegistryOrTableCustomization.h"
+#include "DetailWidgetRow.h"
+#include "PropertyCustomizationHelpers.h"
+#include "DataRegistryEditorModule.h"
+#include "IDetailChildrenBuilder.h"
+#include "DataRegistryIdCustomization.h"
+#include "DataRegistrySubsystem.h"
+#include "PropertyCustomizationHelpers.h"
+#include "Widgets/Images/SImage.h"
+#include "Widgets/Input/SButton.h"
+
+#define LOCTEXT_NAMESPACE "SoftDataRegistryOrTable"
+
+// Makes a new instance of this detail layout class for a specific detail view requesting it
+TSharedRef<IPropertyTypeCustomization> FSoftDataRegistryOrTableCustomization::MakeInstance()
+{
+	return MakeShareable(new FSoftDataRegistryOrTableCustomization);
+
+}
+
+void FSoftDataRegistryOrTableCustomization::CustomizeHeader(TSharedRef<class IPropertyHandle> InStructPropertyHandle, class FDetailWidgetRow& HeaderRow, IPropertyTypeCustomizationUtils& StructCustomizationUtils)
+{
+	this->StructPropertyHandle = InStructPropertyHandle;
+	
+	FilterStructName = NAME_None;
+	if (StructPropertyHandle->HasMetaData(FDataRegistryType::ItemStructMetaData))
+	{
+		const FString& RowType = StructPropertyHandle->GetMetaData(FDataRegistryType::ItemStructMetaData);
+		FilterStructName = FName(*RowType);
+	}
+}
+
+void FSoftDataRegistryOrTableCustomization::GetDataSourceComboStrings(TArray< TSharedPtr<FString> >& OutComboBoxStrings, TArray<TSharedPtr<SToolTip>>& OutToolTips, TArray<bool>& OutRestrictedItems) const
+{
+	OutComboBoxStrings.Add(MakeShared<FString>(LOCTEXT("DataTable", "Data Table").ToString()));
+	OutComboBoxStrings.Add(MakeShared<FString>(LOCTEXT("DataRegistry", "Data Registry").ToString()));
+	OutRestrictedItems.Add(false);
+	OutRestrictedItems.Add(false);	
+}
+
+bool FSoftDataRegistryOrTableCustomization::ShouldFilterDataTableAsset(const struct FAssetData& AssetData)
+{
+	if (!FilterStructName.IsNone())
+	{
+		static const FName RowStructureTagName("RowStructure");
+		UScriptStruct* RowFilterStruct = UClass::TryFindTypeSlow<UScriptStruct>(FilterStructName.ToString());
+		FString RowStructure;
+		if (AssetData.GetTagValue<FString>(RowStructureTagName, RowStructure))
+		{
+			// This is slow, but at the moment we don't have an alternative to the short struct name search
+			UScriptStruct* RowStruct = UClass::TryFindTypeSlow<UScriptStruct>(RowStructure);
+
+			if (RowStruct && RowFilterStruct && RowStruct->IsChildOf(RowFilterStruct))
+			{
+				return false;
+			}
+		}
+		return true;
+	}
+	return false;
+}
+
+FString FSoftDataRegistryOrTableCustomization::GetDataSourceTypeValueString() const
+{
+	if (!UseRegistryPropertyHandle.IsValid() || !UseRegistryPropertyHandle->IsValidHandle())
+	{
+		return LOCTEXT("NoneValue", "None").ToString();	
+	}
+
+	bool bUseDataRegistryId = false;
+	const FPropertyAccess::Result SourceResult = UseRegistryPropertyHandle->GetValue(bUseDataRegistryId);
+	if (SourceResult == FPropertyAccess::Success)
+	{
+		if (bUseDataRegistryId)
+		{
+			return LOCTEXT("DataRegistry", "Data Registry").ToString();
+		}
+		else
+		{
+			return LOCTEXT("DataTable", "Data Table").ToString();	
+		}
+	}
+	else if (SourceResult == FPropertyAccess::Fail)
+	{
+		return LOCTEXT("DataTable_None", "None").ToString();
+	}
+	else
+	{
+		return LOCTEXT("MultipleValues", "Multiple Values").ToString();
+	}
+}
+
+void FSoftDataRegistryOrTableCustomization::OnDataSourceSelected(const FString& String)
+{
+	const bool bUsingDataRegistry = String == LOCTEXT("DataRegistry", "Data Registry").ToString();
+
+	static auto CheckHandle = [](const TSharedPtr<IPropertyHandle>& Handle)
+	{
+		return Handle.IsValid() && Handle->IsValidHandle();
+	};
+
+	if (CheckHandle(UseRegistryPropertyHandle))
+	{
+		UseRegistryPropertyHandle->SetValue(bUsingDataRegistry);
+	}
+
+	if (bUsingDataRegistry)
+	{
+		// clear the data table values to prevent hidden references to a resource
+		if (CheckHandle(TablePropertyHandle))
+		{
+			TablePropertyHandle->SetValue(static_cast<UDataTable*>(nullptr));
+		}
+	}
+	else
+	{
+		// clear the data registry values to prevent hidden references to the registry
+		if (CheckHandle(RegistryTypeNamePropertyHandle))
+		{
+			RegistryTypeNamePropertyHandle->SetValue(FName{});
+		}
+	}
+}
+
+FSoftDataRegistryOrTable FSoftDataRegistryOrTableCustomization::GetCurrentValue() const
+{
+	FSoftDataRegistryOrTable Result;
+
+	static auto CheckHandle = [](const TSharedPtr<IPropertyHandle>& Handle)
+	{
+		return Handle.IsValid() && Handle->IsValidHandle();
+	};
+
+	if (!CheckHandle(UseRegistryPropertyHandle) ||
+		!CheckHandle(RegistryTypeNamePropertyHandle) ||
+		!CheckHandle(TablePropertyHandle))
+	{
+		return Result;
+	}
+
+	UseRegistryPropertyHandle->GetValue(Result.bUseDataRegistry);
+
+	if (Result.bUseDataRegistry)
+	{
+		RegistryTypeNamePropertyHandle->GetValue(Result.RegistryType);
+	}
+	else
+	{
+		UObject* SourceDataTable = nullptr;
+
+		TablePropertyHandle->GetValue(SourceDataTable);
+		Result.Table = Cast<UDataTable>(SourceDataTable);
+	}
+
+	return Result;	
+}
+
+FString FSoftDataRegistryOrTableCustomization::OnGetDataRegistryTypeValueString() const
+{
+	const FSoftDataRegistryOrTable CurrentValue = GetCurrentValue();
+	if (CurrentValue.bUseDataRegistry)
+	{
+		return CurrentValue.RegistryType.ToString();
+	}
+
+	return LOCTEXT("NoneValue", "None").ToString();
+}
+
+EVisibility FSoftDataRegistryOrTableCustomization::GetOpenDataRegistryAssetVisibility() const
+{
+	const FSoftDataRegistryOrTable CurrentValue = GetCurrentValue();
+	return (CurrentValue.bUseDataRegistry && CurrentValue.RegistryType.IsValid()) ? EVisibility::Visible : EVisibility::Collapsed;
+}
+
+FReply FSoftDataRegistryOrTableCustomization::OnClickOpenDataRegistryAsset()
+{
+	const FSoftDataRegistryOrTable CurrentValue = GetCurrentValue();
+
+	if (ensure(CurrentValue.bUseDataRegistry))
+	{
+		if (const UDataRegistrySubsystem* DataRegistrySubsystem = UDataRegistrySubsystem::Get())
+		{
+			if (UDataRegistry* DataRegistry = DataRegistrySubsystem->GetRegistryForType(CurrentValue.RegistryType))
+			{
+				GEditor->GetEditorSubsystem<UAssetEditorSubsystem>()->OpenEditorForAsset(DataRegistry);
+			}
+		}
+	}
+
+	return FReply::Handled();
+}
+
+
+FText FSoftDataRegistryOrTableCustomization::GetOpenDataRegistryAssetTooltip() const
+{
+	const FSoftDataRegistryOrTable CurrentValue = GetCurrentValue();
+
+	if (CurrentValue.bUseDataRegistry)
+	{
+		if (const UDataRegistrySubsystem* DataRegistrySubsystem = UDataRegistrySubsystem::Get())
+		{
+			if (UDataRegistry* DataRegistry = DataRegistrySubsystem->GetRegistryForType(CurrentValue.RegistryType))
+			{
+				FFormatNamedArguments Args;
+				Args.Add(TEXT("Asset"), FText::AsCultureInvariant(DataRegistry->GetName()));
+				return FText::Format(LOCTEXT("OpenSpecificDataRegistry", "Open '{Asset}' in the editor"), Args);
+			}
+		}
+	}
+
+	return LOCTEXT("OpenDataRegistry", "Open the Data Registry in the editor");
+}
+
+
+void FSoftDataRegistryOrTableCustomization::CustomizeChildren(TSharedRef<class IPropertyHandle> InStructPropertyHandle, class IDetailChildrenBuilder& StructBuilder, IPropertyTypeCustomizationUtils& StructCustomizationUtils)
+{
+	static auto CheckHandle = [](const TSharedPtr<IPropertyHandle>& Handle)
+	{
+		return Handle.IsValid() && Handle->IsValidHandle();
+	};
+
+	// get the child property handles
+	UseRegistryPropertyHandle = InStructPropertyHandle->GetChildHandle(GET_MEMBER_NAME_CHECKED(FSoftDataRegistryOrTable, bUseDataRegistry));
+	if (!CheckHandle(UseRegistryPropertyHandle))
+	{
+		return;
+	}
+
+	TSharedPtr<IPropertyHandle> RegistryTypePropertyHandle = InStructPropertyHandle->GetChildHandle(GET_MEMBER_NAME_CHECKED(FSoftDataRegistryOrTable, RegistryType));
+	if (!CheckHandle(RegistryTypePropertyHandle))
+	{
+		return;
+	}
+
+	RegistryTypeNamePropertyHandle = RegistryTypePropertyHandle->GetChildHandle(FName("Name"));
+	if (!CheckHandle(RegistryTypeNamePropertyHandle))
+	{
+		return;
+	}
+	
+	TablePropertyHandle = InStructPropertyHandle->GetChildHandle(GET_MEMBER_NAME_CHECKED(FSoftDataRegistryOrTable, Table));
+	if (!CheckHandle(TablePropertyHandle))
+	{
+		return;
+	}
+
+	FPropertyComboBoxArgs SelectSourceTypeArgs(UseRegistryPropertyHandle,
+		FOnGetPropertyComboBoxStrings::CreateSP(this, &FSoftDataRegistryOrTableCustomization::GetDataSourceComboStrings),
+		FOnGetPropertyComboBoxValue::CreateSP(this, &FSoftDataRegistryOrTableCustomization::GetDataSourceTypeValueString),
+		FOnPropertyComboBoxValueSelected::CreateSP(this, &FSoftDataRegistryOrTableCustomization::OnDataSourceSelected)
+	);
+
+	FPropertyComboBoxArgs DataRegistryIdTypeArgs(RegistryTypeNamePropertyHandle,
+		FOnGetPropertyComboBoxStrings::CreateStatic(FDataRegistryEditorModule::GenerateDataRegistryTypeComboBoxStrings, true, FilterStructName),
+		FOnGetPropertyComboBoxValue::CreateSP(this, &FSoftDataRegistryOrTableCustomization::OnGetDataRegistryTypeValueString));
+	DataRegistryIdTypeArgs.ShowSearchForItemCount = 1;
+
+	auto IsUsingDataRegistry = [this]()
+	{
+		bool bResult = false;
+		if (UseRegistryPropertyHandle.IsValid() && UseRegistryPropertyHandle->IsValidHandle())
+		{
+			UseRegistryPropertyHandle->GetValue(bResult);
+		}
+		return bResult ? EVisibility::Visible : EVisibility::Collapsed;
+	};
+
+	auto IsUsingDataTable = [this]()
+	{
+		bool bResult = false;
+		if (UseRegistryPropertyHandle.IsValid() && UseRegistryPropertyHandle->IsValidHandle())
+		{
+			UseRegistryPropertyHandle->GetValue(bResult);
+		}
+		return bResult ? EVisibility::Collapsed : EVisibility::Visible;
+	};
+
+
+	FUIAction CopyAction;
+	FUIAction PasteAction;
+	InStructPropertyHandle->CreateDefaultPropertyCopyPasteActions(CopyAction, PasteAction);
+
+	StructBuilder.AddCustomRow(LOCTEXT("DataTableOrRegistry", "Data Table Or Registry"))
+		.CopyAction(CopyAction)
+		.PasteAction(PasteAction)
+		.NameContent()
+		.HAlign(HAlign_Fill)
+		[
+			InStructPropertyHandle->CreatePropertyNameWidget()
+		]
+		.ValueContent()
+		.MinDesiredWidth(800.f)
+		[
+			SNew(SWrapBox)
+			.UseAllottedSize(true)
+			.Orientation(Orient_Horizontal)
+			.HAlign(HAlign_Fill)			
+			+ SWrapBox::Slot()
+				.FillEmptySpace(false)
+				[
+					PropertyCustomizationHelpers::MakePropertyComboBox(SelectSourceTypeArgs)
+				]
+			+ SWrapBox::Slot()
+				.FillEmptySpace(true)
+				[
+					SNew(SObjectPropertyEntryBox)
+					.PropertyHandle(TablePropertyHandle)
+					.AllowedClass(UDataTable::StaticClass())
+					.OnShouldFilterAsset(this, &FSoftDataRegistryOrTableCustomization::ShouldFilterDataTableAsset)
+					.Visibility(TAttribute<EVisibility>::Create(IsUsingDataTable))
+				]
+			+ SWrapBox::Slot()
+				.FillEmptySpace(true)
+				[
+					SNew(SHorizontalBox)
+					.Visibility(TAttribute<EVisibility>::Create(IsUsingDataRegistry))
+					+ SHorizontalBox::Slot()
+						.AutoWidth()
+						[
+							PropertyCustomizationHelpers::MakePropertyComboBox(DataRegistryIdTypeArgs)
+						]
+					+ SHorizontalBox::Slot()
+						.AutoWidth()
+						.HAlign(HAlign_Center)
+						.VAlign(VAlign_Center)
+						.Padding(1, 0)
+						[
+							SNew(SBox)
+							.HAlign(HAlign_Center)
+							.VAlign(VAlign_Center)
+							.WidthOverride(22.0f)
+							.HeightOverride(22.0f)
+							.ToolTipText(this, &FSoftDataRegistryOrTableCustomization::GetOpenDataRegistryAssetTooltip)
+							[
+								SNew(SButton)
+								.ButtonStyle(FAppStyle::Get(), "SimpleButton")
+								.OnClicked(this, &FSoftDataRegistryOrTableCustomization::OnClickOpenDataRegistryAsset)
+								.ContentPadding(0.0f)
+								.IsFocusable(false)
+								.Visibility(this, &FSoftDataRegistryOrTableCustomization::GetOpenDataRegistryAssetVisibility)
+								[
+									SNew(SImage)
+										.Image(FAppStyle::GetBrush("SystemWideCommands.SummonOpenAssetDialog"))
+										.ColorAndOpacity(FSlateColor::UseForeground())
+								]
+							]
+						]
+				]			
+		]
+		.OverrideResetToDefault(FResetToDefaultOverride::Create(
+			TAttribute<bool>::CreateLambda([InStructPropertyHandle]() -> bool
+			{
+				return InStructPropertyHandle->IsValidHandle() && InStructPropertyHandle->DiffersFromDefault();
+			}),
+			FSimpleDelegate::CreateLambda([InStructPropertyHandle]()
+			{
+				if (InStructPropertyHandle->IsValidHandle())
+				{
+					InStructPropertyHandle->ResetToDefault();
+				}
+			})
+		)
+	);	
+}
+
diff --git a/Engine/Plugins/Runtime/DataRegistry/Source/DataRegistryEditor/Private/SoftDataRegistryOrTableCustomization.h b/Engine/Plugins/Runtime/DataRegistry/Source/DataRegistryEditor/Private/SoftDataRegistryOrTableCustomization.h
new file mode 100644
--- /dev/null
+++ b/Engine/Plugins/Runtime/DataRegistry/Source/DataRegistryEditor/Private/SoftDataRegistryOrTableCustomization.h
@@ -0,0 +1,41 @@
+// Copyright Epic Games, Inc. All Rights Reserved.
+
+#pragma once
+
+#include "SoftDataRegistryOrTable.h"
+#include "IPropertyTypeCustomization.h"
+#include "Widgets/SToolTip.h"
+
+/**
+* Customizes a FSoftDataRegistryOrTable to support filtering by item struct, and compacting it to one row 
+  */
+class FSoftDataRegistryOrTableCustomization : public IPropertyTypeCustomization
+{
+public:
+	// Makes a new instance of this detail layout class for a specific detail view requesting it
+	static TSharedRef<IPropertyTypeCustomization> MakeInstance();
+
+	/** IPropertyTypeCustomization interface */
+	virtual void CustomizeHeader(TSharedRef<class IPropertyHandle> InStructPropertyHandle, class FDetailWidgetRow& HeaderRow, IPropertyTypeCustomizationUtils& StructCustomizationUtils) override;
+	virtual void CustomizeChildren(TSharedRef<class IPropertyHandle> InStructPropertyHandle, class IDetailChildrenBuilder& StructBuilder, IPropertyTypeCustomizationUtils& StructCustomizationUtils) override;
+
+private:
+	void GetDataSourceComboStrings(TArray< TSharedPtr<FString> >& OutComboBoxStrings, TArray<TSharedPtr<SToolTip>>& OutToolTips, TArray<bool>& OutRestrictedItems) const;
+	FString GetDataSourceTypeValueString() const;
+	void OnDataSourceSelected(const FString& String);
+	FString OnGetDataRegistryTypeValueString() const;
+	bool ShouldFilterDataTableAsset(const struct FAssetData& AssetData);
+
+	EVisibility GetOpenDataRegistryAssetVisibility() const;
+	FReply OnClickOpenDataRegistryAsset();
+	FText GetOpenDataRegistryAssetTooltip() const;
+
+	FSoftDataRegistryOrTable GetCurrentValue() const;
+
+	TSharedPtr<IPropertyHandle> StructPropertyHandle;
+	TSharedPtr<IPropertyHandle> UseRegistryPropertyHandle;
+	TSharedPtr<IPropertyHandle> RegistryTypeNamePropertyHandle;
+	TSharedPtr<IPropertyHandle> TablePropertyHandle;
+
+	FName FilterStructName;	
+};
