Fixing a bughawk issue

diff --git a/Engine/Plugins/Runtime/DataRegistry/Source/DataRegistryEditor/Private/SoftDataRegistryOrTableCustomization.cpp b/Engine/Plugins/Runtime/DataRegistry/Source/DataRegistryEditor/Private/SoftDataRegistryOrTableCustomization.cpp
--- a/Engine/Plugins/Runtime/DataRegistry/Source/DataRegistryEditor/Private/SoftDataRegistryOrTableCustomization.cpp
+++ b/Engine/Plugins/Runtime/DataRegistry/Source/DataRegistryEditor/Private/SoftDataRegistryOrTableCustomization.cpp
@@ -365,3 +365,4 @@ void FSoftDataRegistryOrTableCustomization::CustomizeChildren(TSharedRef<class I
 	);	
 }
 
+#undef LOCTEXT_NAMESPACE
