Fix for a crash bug in DataTableOrRegistyRowCustomization and SoftDataTableRegistryOrTableCustomization.
Optimization for their struct type metadata param

diff --git a/Engine/Plugins/Runtime/DataRegistry/Source/DataRegistryEditor/Private/DataTableOrRegistryRowCustomization.cpp b/Engine/Plugins/Runtime/DataRegistry/Source/DataRegistryEditor/Private/DataTableOrRegistryRowCustomization.cpp
--- a/Engine/Plugins/Runtime/DataRegistry/Source/DataRegistryEditor/Private/DataTableOrRegistryRowCustomization.cpp
+++ b/Engine/Plugins/Runtime/DataRegistry/Source/DataRegistryEditor/Private/DataTableOrRegistryRowCustomization.cpp
@@ -24,12 +24,25 @@ TSharedRef<IPropertyTypeCustomization> FDataRegistryOrTableRowCustomization::Mak
 void FDataRegistryOrTableRowCustomization::CustomizeHeader(TSharedRef<class IPropertyHandle> InStructPropertyHandle, class FDetailWidgetRow& HeaderRow, IPropertyTypeCustomizationUtils& StructCustomizationUtils)
 {
 	this->StructPropertyHandle = InStructPropertyHandle;
-	
+
 	FilterStructName = NAME_None;
+	RegistryFilterStructName = NAME_None;
+
 	if (StructPropertyHandle->HasMetaData(FDataRegistryType::ItemStructMetaData))
 	{
 		const FString& RowType = StructPropertyHandle->GetMetaData(FDataRegistryType::ItemStructMetaData);
 		FilterStructName = FName(*RowType);
+
+		// The registry filter needs the short name, but the data table is very slow unless the long name is used, so just extract the suffix
+		int32 ExtIndex = -1;
+		if (RowType.FindLastChar('.', ExtIndex))
+		{
+			RegistryFilterStructName = FName(RowType.Mid(ExtIndex+1));
+		}
+		else
+		{
+			RegistryFilterStructName = FilterStructName;
+		}		
 	}
 }
 
@@ -386,7 +399,7 @@ void FDataRegistryOrTableRowCustomization::CustomizeChildren(TSharedRef<class IP
 	);
 
 	FPropertyComboBoxArgs DataRegistryIdTypeArgs(RegistryTypeNamePropertyHandle,
-		FOnGetPropertyComboBoxStrings::CreateStatic(FDataRegistryEditorModule::GenerateDataRegistryTypeComboBoxStrings, true, FilterStructName),
+		FOnGetPropertyComboBoxStrings::CreateStatic(FDataRegistryEditorModule::GenerateDataRegistryTypeComboBoxStrings, true, RegistryFilterStructName),
 		FOnGetPropertyComboBoxValue::CreateSP(this, &FDataRegistryOrTableRowCustomization::OnGetDataRegistryTypeValueString));
 	DataRegistryIdTypeArgs.ShowSearchForItemCount = 1;
 
@@ -395,9 +408,12 @@ void FDataRegistryOrTableRowCustomization::CustomizeChildren(TSharedRef<class IP
 		FOnGetPropertyComboBoxValue::CreateSP(this, &FDataRegistryOrTableRowCustomization::OnGetDataTableRowValueString));
 	DataTableRowStringsArgs.ShowSearchForItemCount = 1;
 
-	auto IsUsingDataRegistry = [this]()
+	auto IsUsingDataRegistry = [PropertyHandleWP = UseRegistryPropertyHandle.ToWeakPtr()]()
 	{
 		bool bResult = false;
+
+		TSharedPtr<IPropertyHandle> UseRegistryPropertyHandle = PropertyHandleWP.Pin();
+
 		if (UseRegistryPropertyHandle.IsValid() && UseRegistryPropertyHandle->IsValidHandle())
 		{
 			UseRegistryPropertyHandle->GetValue(bResult);
@@ -405,9 +421,12 @@ void FDataRegistryOrTableRowCustomization::CustomizeChildren(TSharedRef<class IP
 		return bResult ? EVisibility::Visible : EVisibility::Collapsed;
 	};
 
-	auto IsUsingDataTable = [this]()
+	auto IsUsingDataTable = [PropertyHandleWP = UseRegistryPropertyHandle.ToWeakPtr()]()
 	{
 		bool bResult = false;
+
+		TSharedPtr<IPropertyHandle> UseRegistryPropertyHandle = PropertyHandleWP.Pin();
+
 		if (UseRegistryPropertyHandle.IsValid() && UseRegistryPropertyHandle->IsValidHandle())
 		{
 			UseRegistryPropertyHandle->GetValue(bResult);
@@ -415,7 +434,7 @@ void FDataRegistryOrTableRowCustomization::CustomizeChildren(TSharedRef<class IP
 		return bResult ? EVisibility::Collapsed : EVisibility::Visible;
 	};
 
-	auto IsUsingDataTableAndSourceSelected = [this, IsUsingDataTable]() {
+	auto IsUsingDataTableAndSourceSelected = [TablePropertyHandleWP = TablePropertyHandle.ToWeakPtr(), IsUsingDataTable]() {
 
 		if (IsUsingDataTable() == EVisibility::Collapsed) 
 		{
@@ -423,6 +442,9 @@ void FDataRegistryOrTableRowCustomization::CustomizeChildren(TSharedRef<class IP
 		}
 
 		UObject* SourceDataTable = nullptr;
+
+		TSharedPtr<IPropertyHandle> TablePropertyHandle = TablePropertyHandleWP.Pin();
+
 		if (TablePropertyHandle.IsValid() && TablePropertyHandle->IsValidHandle())
 		{
 			TablePropertyHandle->GetValue(SourceDataTable);
@@ -431,7 +453,7 @@ void FDataRegistryOrTableRowCustomization::CustomizeChildren(TSharedRef<class IP
 		return SourceDataTable != nullptr ? EVisibility::Visible : EVisibility::Collapsed;
 	};
 
-	auto IsUsingDataRegistryAndSourceSelected = [this, IsUsingDataRegistry]() {
+	auto IsUsingDataRegistryAndSourceSelected = [RegistryTypeNamePropertyHandleWP = RegistryTypeNamePropertyHandle.ToWeakPtr(), IsUsingDataRegistry]() {
 
 		if (IsUsingDataRegistry() == EVisibility::Collapsed) 
 		{
@@ -440,6 +462,8 @@ void FDataRegistryOrTableRowCustomization::CustomizeChildren(TSharedRef<class IP
 
 		FName RegistryType;
 
+		TSharedPtr<IPropertyHandle> RegistryTypeNamePropertyHandle = RegistryTypeNamePropertyHandleWP.Pin();
+
 		RegistryTypeNamePropertyHandle->GetValue(RegistryType);
 
 		return RegistryType.IsNone() ? EVisibility::Collapsed : EVisibility::Visible;
diff --git a/Engine/Plugins/Runtime/DataRegistry/Source/DataRegistryEditor/Private/DataTableOrRegistryRowCustomization.h b/Engine/Plugins/Runtime/DataRegistry/Source/DataRegistryEditor/Private/DataTableOrRegistryRowCustomization.h
--- a/Engine/Plugins/Runtime/DataRegistry/Source/DataRegistryEditor/Private/DataTableOrRegistryRowCustomization.h
+++ b/Engine/Plugins/Runtime/DataRegistry/Source/DataRegistryEditor/Private/DataTableOrRegistryRowCustomization.h
@@ -42,6 +42,7 @@ private:
 	TSharedPtr<IPropertyHandle> StructPropertyHandle;
 
 	FName FilterStructName;
+	FName RegistryFilterStructName;
 
 	TSharedPtr<IPropertyHandle> UseRegistryPropertyHandle;
 	TSharedPtr<IPropertyHandle> RegistryTypeNamePropertyHandle;
diff --git a/Engine/Plugins/Runtime/DataRegistry/Source/DataRegistryEditor/Private/SoftDataRegistryOrTableCustomization.cpp b/Engine/Plugins/Runtime/DataRegistry/Source/DataRegistryEditor/Private/SoftDataRegistryOrTableCustomization.cpp
--- a/Engine/Plugins/Runtime/DataRegistry/Source/DataRegistryEditor/Private/SoftDataRegistryOrTableCustomization.cpp
+++ b/Engine/Plugins/Runtime/DataRegistry/Source/DataRegistryEditor/Private/SoftDataRegistryOrTableCustomization.cpp
@@ -25,10 +25,23 @@ void FSoftDataRegistryOrTableCustomization::CustomizeHeader(TSharedRef<class IPr
 	this->StructPropertyHandle = InStructPropertyHandle;
 	
 	FilterStructName = NAME_None;
+	RegistryFilterStructName = NAME_None;
+
 	if (StructPropertyHandle->HasMetaData(FDataRegistryType::ItemStructMetaData))
 	{
 		const FString& RowType = StructPropertyHandle->GetMetaData(FDataRegistryType::ItemStructMetaData);
 		FilterStructName = FName(*RowType);
+
+		// The registry filter needs the short name, but the data table is very slow unless the long name is used, so just extract the suffix
+		int32 ExtIndex = -1;
+		if (RowType.FindLastChar('.', ExtIndex))
+		{
+			RegistryFilterStructName = FName(RowType.Mid(ExtIndex+1));
+		}
+		else
+		{
+			RegistryFilterStructName = FilterStructName;
+		}
 	}
 }
 
@@ -253,26 +266,31 @@ void FSoftDataRegistryOrTableCustomization::CustomizeChildren(TSharedRef<class I
 	);
 
 	FPropertyComboBoxArgs DataRegistryIdTypeArgs(RegistryTypeNamePropertyHandle,
-		FOnGetPropertyComboBoxStrings::CreateStatic(FDataRegistryEditorModule::GenerateDataRegistryTypeComboBoxStrings, true, FilterStructName),
+		FOnGetPropertyComboBoxStrings::CreateStatic(FDataRegistryEditorModule::GenerateDataRegistryTypeComboBoxStrings, true, RegistryFilterStructName),
 		FOnGetPropertyComboBoxValue::CreateSP(this, &FSoftDataRegistryOrTableCustomization::OnGetDataRegistryTypeValueString));
 	DataRegistryIdTypeArgs.ShowSearchForItemCount = 1;
 
-	auto IsUsingDataRegistry = [this]()
+	auto IsUsingDataRegistry = [PropertyHandleWP = UseRegistryPropertyHandle.ToWeakPtr()]()
 	{
 		bool bResult = false;
-		if (UseRegistryPropertyHandle.IsValid() && UseRegistryPropertyHandle->IsValidHandle())
+
+		TSharedPtr<IPropertyHandle> PropertyHandle = PropertyHandleWP.Pin();
+
+		if (PropertyHandle.IsValid() && PropertyHandle->IsValidHandle())
 		{
-			UseRegistryPropertyHandle->GetValue(bResult);
+			PropertyHandle->GetValue(bResult);
 		}
 		return bResult ? EVisibility::Visible : EVisibility::Collapsed;
 	};
 
-	auto IsUsingDataTable = [this]()
+	auto IsUsingDataTable = [PropertyHandleWP = UseRegistryPropertyHandle.ToWeakPtr()]()
 	{
 		bool bResult = false;
-		if (UseRegistryPropertyHandle.IsValid() && UseRegistryPropertyHandle->IsValidHandle())
+		TSharedPtr<IPropertyHandle> PropertyHandle = PropertyHandleWP.Pin();
+
+		if (PropertyHandle.IsValid() && PropertyHandle->IsValidHandle())
 		{
-			UseRegistryPropertyHandle->GetValue(bResult);
+			PropertyHandle->GetValue(bResult);
 		}
 		return bResult ? EVisibility::Collapsed : EVisibility::Visible;
 	};
diff --git a/Engine/Plugins/Runtime/DataRegistry/Source/DataRegistryEditor/Private/SoftDataRegistryOrTableCustomization.h b/Engine/Plugins/Runtime/DataRegistry/Source/DataRegistryEditor/Private/SoftDataRegistryOrTableCustomization.h
--- a/Engine/Plugins/Runtime/DataRegistry/Source/DataRegistryEditor/Private/SoftDataRegistryOrTableCustomization.h
+++ b/Engine/Plugins/Runtime/DataRegistry/Source/DataRegistryEditor/Private/SoftDataRegistryOrTableCustomization.h
@@ -38,4 +38,5 @@ private:
 	TSharedPtr<IPropertyHandle> TablePropertyHandle;
 
 	FName FilterStructName;	
+	FName RegistryFilterStructName;	
 };
