Fix for bughawk issue

diff --git a/Engine/Plugins/Runtime/DataRegistry/Source/DataRegistryEditor/Private/DataTableOrRegistryRowCustomization.cpp b/Engine/Plugins/Runtime/DataRegistry/Source/DataRegistryEditor/Private/DataTableOrRegistryRowCustomization.cpp
--- a/Engine/Plugins/Runtime/DataRegistry/Source/DataRegistryEditor/Private/DataTableOrRegistryRowCustomization.cpp
+++ b/Engine/Plugins/Runtime/DataRegistry/Source/DataRegistryEditor/Private/DataTableOrRegistryRowCustomization.cpp
@@ -460,11 +460,14 @@ void FDataRegistryOrTableRowCustomization::CustomizeChildren(TSharedRef<class IP
 			return EVisibility::Collapsed;
 		}
 
-		FName RegistryType;
+		FName RegistryType = NAME_None;
 
 		TSharedPtr<IPropertyHandle> RegistryTypeNamePropertyHandle = RegistryTypeNamePropertyHandleWP.Pin();
 
-		RegistryTypeNamePropertyHandle->GetValue(RegistryType);
+		if (RegistryTypeNamePropertyHandle.IsValid() && RegistryTypeNamePropertyHandle->IsValidHandle())
+		{
+			RegistryTypeNamePropertyHandle->GetValue(RegistryType);
+		}
 
 		return RegistryType.IsNone() ? EVisibility::Collapsed : EVisibility::Visible;
 	};
