[DataRegistry] Fixed "Find Row References" to view references with type FDataRegistryId::StaticStruct() using the name from FDataRegistryId.ToString().

Swarm note: This CL was reviewed, tested, and voted up under a different branch and swarm-review 48526630, and then moved to main where it got submitted.

#tests Verified in PIE against a temp local registry and temp local BP structure content with a DataRegistryId.
#Synced-CL 48809907

diff --git a/Engine/Plugins/Runtime/DataRegistry/Source/DataRegistry/Private/DataRegistryId.cpp b/Engine/Plugins/Runtime/DataRegistry/Source/DataRegistry/Private/DataRegistryId.cpp
--- a/Engine/Plugins/Runtime/DataRegistry/Source/DataRegistry/Private/DataRegistryId.cpp
+++ b/Engine/Plugins/Runtime/DataRegistry/Source/DataRegistry/Private/DataRegistryId.cpp
@@ -111,6 +111,14 @@ bool FDataRegistryId::SerializeFromMismatchedTag(struct FPropertyTag const& Tag,
 	return false;
 }
 
+void FDataRegistryId::PostSerialize(const FArchive& Ar)
+{
+	if (Ar.IsSaving() && IsValid())
+	{
+		Ar.MarkSearchableName(FDataRegistryId::StaticStruct(), FName(*ToString()));
+	}
+}
+
 FString FDataRegistryId::ToString() const
 {
 	if (IsValid())
diff --git a/Engine/Plugins/Runtime/DataRegistry/Source/DataRegistry/Public/DataRegistryId.h b/Engine/Plugins/Runtime/DataRegistry/Source/DataRegistry/Public/DataRegistryId.h
--- a/Engine/Plugins/Runtime/DataRegistry/Source/DataRegistry/Public/DataRegistryId.h
+++ b/Engine/Plugins/Runtime/DataRegistry/Source/DataRegistry/Public/DataRegistryId.h
@@ -193,6 +193,7 @@ struct FDataRegistryId
 	DATAREGISTRY_API bool ExportTextItem(FString& ValueStr, FDataRegistryId const& DefaultValue, UObject* Parent, int32 PortFlags, UObject* ExportRootScope) const;
 	DATAREGISTRY_API bool ImportTextItem(const TCHAR*& Buffer, int32 PortFlags, UObject* Parent, FOutputDevice* ErrorText);
 	DATAREGISTRY_API bool SerializeFromMismatchedTag(struct FPropertyTag const& Tag, FStructuredArchive::FSlot Slot);
+	DATAREGISTRY_API void PostSerialize(const FArchive& Ar); 
 
 	friend inline uint32 GetTypeHash(const FDataRegistryId& Key)
 	{
@@ -215,5 +216,6 @@ struct TStructOpsTypeTraits<FDataRegistryId> : public TStructOpsTypeTraitsBase2<
 		WithExportTextItem = true,
 		WithImportTextItem = true,
 		WithStructuredSerializeFromMismatchedTag = true,
+		WithPostSerialize = true,
 	};
 };
diff --git a/Engine/Plugins/Runtime/DataRegistry/Source/DataRegistryEditor/Private/SDataRegistryListViewRow.cpp b/Engine/Plugins/Runtime/DataRegistry/Source/DataRegistryEditor/Private/SDataRegistryListViewRow.cpp
--- a/Engine/Plugins/Runtime/DataRegistry/Source/DataRegistryEditor/Private/SDataRegistryListViewRow.cpp
+++ b/Engine/Plugins/Runtime/DataRegistry/Source/DataRegistryEditor/Private/SDataRegistryListViewRow.cpp
@@ -60,15 +60,13 @@ void SDataRegistryListViewRow::OnSearchForReferences()
 {
 	if (const FDataRegistrySourceItemId* FoundSource = FindSource())
 	{
-		// TODO Either properly export this in serialize or disable this menu option, similar code in customization
 		TArray<FAssetIdentifier> AssetIdentifiers;
-		AssetIdentifiers.Add(FAssetIdentifier(FoundSource->ItemId.RegistryType, FoundSource->ItemId.ItemName));
+		AssetIdentifiers.Add(FAssetIdentifier(FDataRegistryId::StaticStruct(), FName(FoundSource->ItemId.ToString())));
 
 		FEditorDelegates::OnOpenReferenceViewer.Broadcast(AssetIdentifiers, FReferenceViewerParams());
 	}
 }
 
-
 void SDataRegistryListViewRow::OnOpenSourceAsset()
 {
 	const FDataRegistrySourceItemId* FoundSource = FindSource();
