Fix issue where trying to unregister an asset from a temporarily-unregistered DataRegistry would fail, causing GC leaks when the data registry is reinitialized later

#Synced-CL 50988622

diff --git a/Engine/Plugins/Runtime/DataRegistry/Source/DataRegistry/Private/DataRegistrySubsystem.cpp b/Engine/Plugins/Runtime/DataRegistry/Source/DataRegistry/Private/DataRegistrySubsystem.cpp
--- a/Engine/Plugins/Runtime/DataRegistry/Source/DataRegistry/Private/DataRegistrySubsystem.cpp
+++ b/Engine/Plugins/Runtime/DataRegistry/Source/DataRegistry/Private/DataRegistrySubsystem.cpp
@@ -599,9 +599,10 @@ bool UDataRegistrySubsystem::UnregisterSpecificAsset(FDataRegistryType RegistryT
 		}
 	}
 
+	bool bMadeChange = false;
+
 	if (!RegistryType.IsValid())
 	{
-		bool bMadeChange = false;
 		for (const FRegistryMapPair& RegistryPair : RegistryMap)
 		{
 			UDataRegistry* Registry = RegistryPair.Value;
@@ -610,16 +611,30 @@ bool UDataRegistrySubsystem::UnregisterSpecificAsset(FDataRegistryType RegistryT
 				bMadeChange |= Registry->UnregisterSpecificAsset(AssetPath);
 			}
 		}
-		return bMadeChange;
+	}
+	else
+	{
+		UDataRegistry* FoundRegistry = GetRegistryForType(RegistryType);
+		if (FoundRegistry)
+		{
+			bMadeChange = FoundRegistry->UnregisterSpecificAsset(AssetPath);
+		}
 	}
 
-	UDataRegistry* FoundRegistry = GetRegistryForType(RegistryType);
-	if (FoundRegistry)
+	if (!bMadeChange)
 	{
-		return FoundRegistry->UnregisterSpecificAsset(AssetPath);
+		// Data registry may have already unregistered, look at the other ones in memory that are not CDOs
+		for (TObjectIterator<UDataRegistry> RegistryIterator; RegistryIterator; ++RegistryIterator)
+		{
+			UDataRegistry* FoundRegistry = *RegistryIterator;
+			if (!FoundRegistry->IsInitialized() && (!RegistryType.IsValid() || FoundRegistry->GetRegistryType() == RegistryType))
+			{
+				bMadeChange = FoundRegistry->UnregisterSpecificAsset(AssetPath);
+			}
+		}
 	}
 
-	return false;
+	return bMadeChange;
 }
 
 int32 UDataRegistrySubsystem::UnregisterAssetsWithPriority(FDataRegistryType RegistryType, int32 AssetPriority)
