Prevent crash when attempting to create a FSoftDataRegistryOrTable with bad data

#Synced-CL 51472625

#Synced-CL 51472625

diff --git a/Engine/Plugins/Runtime/DataRegistry/Source/DataRegistry/Private/SoftDataRegistryOrTable.cpp b/Engine/Plugins/Runtime/DataRegistry/Source/DataRegistry/Private/SoftDataRegistryOrTable.cpp
--- a/Engine/Plugins/Runtime/DataRegistry/Source/DataRegistry/Private/SoftDataRegistryOrTable.cpp
+++ b/Engine/Plugins/Runtime/DataRegistry/Source/DataRegistry/Private/SoftDataRegistryOrTable.cpp
@@ -22,11 +22,17 @@ FSoftDataRegistryOrTable::FSoftDataRegistryOrTable(const UDataTable* InDataTable
 		Table = nullptr;
 		bUseDataRegistry = true;
 	}
-	else
+	else if(ensure(InDataTable))
 	{ 
 		Table = const_cast<UDataTable*>(InDataTable);
 		bUseDataRegistry = false;
 	}
+	else
+	{
+		RegistryType = NAME_None;
+		Table = nullptr;
+		bUseDataRegistry = false;
+	}
 }
 
 bool FSoftDataRegistryOrTable::Serialize(FArchive& Ar)
