Removing ensure triggering with intentionally null data

Triggering frequently in PIE

#Synced-CL 51472625

#Synced-CL 51472625

diff --git a/Engine/Plugins/Runtime/DataRegistry/Source/DataRegistry/Private/SoftDataRegistryOrTable.cpp b/Engine/Plugins/Runtime/DataRegistry/Source/DataRegistry/Private/SoftDataRegistryOrTable.cpp
--- a/Engine/Plugins/Runtime/DataRegistry/Source/DataRegistry/Private/SoftDataRegistryOrTable.cpp
+++ b/Engine/Plugins/Runtime/DataRegistry/Source/DataRegistry/Private/SoftDataRegistryOrTable.cpp
@@ -22,7 +22,7 @@ FSoftDataRegistryOrTable::FSoftDataRegistryOrTable(const UDataTable* InDataTable
 		Table = nullptr;
 		bUseDataRegistry = true;
 	}
-	else if(ensure(InDataTable))
+	else if(InDataTable)
 	{ 
 		Table = const_cast<UDataTable*>(InDataTable);
 		bUseDataRegistry = false;
